Move Experimental Metasound Plugin out of NotForLicencees

diff --git a/Engine/Plugins/Experimental/MetasoundExperimental/MetasoundExperimental.uplugin b/Engine/Plugins/Experimental/MetasoundExperimental/MetasoundExperimental.uplugin
new file mode 100644
--- /dev/null
+++ b/Engine/Plugins/Experimental/MetasoundExperimental/MetasoundExperimental.uplugin
@@ -0,0 +1,41 @@
+{
+	"FileVersion": 3,
+	"Version": 1,
+	"VersionName": "1.0",
+	"FriendlyName": "Metasounds Experimental",
+	"Description": "Metasound developmental plugin, for new features before they are ready for prime time",
+	"Category": "Audio",
+	"CreatedBy": "Epic Games, Inc.",
+	"CreatedByURL": "http://epicgames.com",
+	"DocsURL": "",
+	"MarketplaceURL": "",
+	"SupportURL": "",
+	"EnabledByDefault": false,
+	"CanContainContent": true,
+	"IsBetaVersion": false,
+	"Installed": false,
+	"Modules": [
+		{
+			"Name": "AudioExperimentalRuntime",
+			"Type": "Runtime",
+			"LoadingPhase": "Default"
+		},
+		{
+			"Name": "MetasoundExperimentalRuntime",
+			"Type": "Runtime",
+			"LoadingPhase": "Default"
+		},
+		{
+			"Name": "MetasoundExperimentalEditor",
+			"Type": "Editor",
+			"LoadingPhase": "Default"
+		}
+	],
+	"Plugins": [
+		{
+			"Name": "Metasound",
+			"Enabled": true
+		}
+	],
+	"IsExperimentalVersion": true	
+}
\ No newline at end of file
diff --git a/Engine/Plugins/Experimental/MetasoundExperimental/Source/AudioExperimentalRuntime/AudioExperimentalRuntime.Build.cs b/Engine/Plugins/Experimental/MetasoundExperimental/Source/AudioExperimentalRuntime/AudioExperimentalRuntime.Build.cs
new file mode 100644
--- /dev/null
+++ b/Engine/Plugins/Experimental/MetasoundExperimental/Source/AudioExperimentalRuntime/AudioExperimentalRuntime.Build.cs
@@ -0,0 +1,33 @@
+﻿// Copyright Epic Games, Inc. All Rights Reserved.
+
+using UnrealBuildTool;
+
+public class AudioExperimentalRuntime : ModuleRules
+{
+	public AudioExperimentalRuntime(ReadOnlyTargetRules Target) : base(Target)
+	{
+		PCHUsage = PCHUsageMode.UseExplicitOrSharedPCHs;
+
+		PublicDependencyModuleNames.AddRange(
+			new string[]
+			{
+				"Core",
+			}
+		);
+		
+		// When the Editor builds all plugins, it will build with engine. In that case add a dependency on CoreUObject.
+		// If we don't do this we will get linker errors for the Engine expecting to find stuff defined here.
+		if (Target.bCompileAgainstEngine)
+		{
+			PublicDependencyModuleNames.Add("CoreUObject");
+		}
+		
+		PrivateDependencyModuleNames.AddRange(
+			new string[]
+			{
+				"Core",
+				"SignalProcessing",
+			}
+		);
+	}
+}
\ No newline at end of file
diff --git a/Engine/Plugins/Experimental/MetasoundExperimental/Source/AudioExperimentalRuntime/Private/AudioExperimentalRuntimeModule.cpp b/Engine/Plugins/Experimental/MetasoundExperimental/Source/AudioExperimentalRuntime/Private/AudioExperimentalRuntimeModule.cpp
new file mode 100644
--- /dev/null
+++ b/Engine/Plugins/Experimental/MetasoundExperimental/Source/AudioExperimentalRuntime/Private/AudioExperimentalRuntimeModule.cpp
@@ -0,0 +1,19 @@
+﻿// Copyright Epic Games, Inc. All Rights Reserved.
+
+#include "CoreMinimal.h"
+#include "Modules/ModuleManager.h"
+#include "TypeFamily/ChannelTypeFamily.h"
+
+class FAudioCoreExperimentalModule final : public IModuleInterface
+{
+public:
+	virtual void StartupModule() override
+	{
+		Audio::RegisterChannelLayouts();
+	}
+	
+	virtual void ShutdownModule() override
+	{}
+};
+    
+IMPLEMENT_MODULE(FAudioCoreExperimentalModule, AudioExperimentalRuntime)
\ No newline at end of file
diff --git a/Engine/Plugins/Experimental/MetasoundExperimental/Source/AudioExperimentalRuntime/Private/ChannelAgnostic/ChannelAgnosticType.cpp b/Engine/Plugins/Experimental/MetasoundExperimental/Source/AudioExperimentalRuntime/Private/ChannelAgnostic/ChannelAgnosticType.cpp
new file mode 100644
--- /dev/null
+++ b/Engine/Plugins/Experimental/MetasoundExperimental/Source/AudioExperimentalRuntime/Private/ChannelAgnostic/ChannelAgnosticType.cpp
@@ -0,0 +1,50 @@
+﻿// Copyright Epic Games, Inc. All Rights Reserved.
+
+
+#include "ChannelAgnostic/ChannelAgnosticType.h"
+#include "SimpleAlloc/SimpleHeapAllocator.h"
+#include "TypeFamily/ChannelTypeFamily.h"
+
+namespace Audio
+{
+	FSimpleAllocBase& FChannelAgnosticType::GetDefaultAllocator()
+	{
+		static FSimpleHeapAllocator Allocator;
+		return Allocator;
+	}
+
+	FChannelAgnosticType::FChannelAgnosticType(const TRetainedRef<const FChannelTypeFamily> InType, const int32 InNumFrames, FSimpleAllocBase* InAllocator)
+		: FChannelAgnosticType(InType, InNumFrames, InType.Get().NumChannels(), InAllocator)
+	{}
+	
+	FChannelAgnosticType::FChannelAgnosticType(const TRetainedRef<const FChannelTypeFamily> InType, const int32 InNumFrames, const int32 InNumChannels, FSimpleAllocBase* InAllocator)
+		: Buffer(InNumFrames * InNumChannels, InAllocator)
+		, Type(&InType.Get())
+		, NumFramesPrivate(InNumFrames)
+		, NumChannelsPrivate(InNumChannels)
+	{
+		check(Type);
+		check(NumChannelsPrivate > 0);
+		check(Type->NumChannels() == InNumChannels || (Type->NumChannels() == 0 && InNumChannels > 0));
+	}
+
+	bool FChannelAgnosticType::IsA(const FChannelAgnosticType& InOther) const
+	{
+		check(Type);
+		check(InOther.Type);
+		return Type->IsA(InOther.Type);
+	}
+
+	bool FChannelAgnosticType::IsA(const FName& InTypeName) const
+	{
+		check(Type);
+		return Type->IsA(InTypeName);
+	}
+	
+	FName FChannelAgnosticType::GetTypeName() const
+	{
+		check(Type);
+		return Type->GetName();
+	}
+
+}
diff --git a/Engine/Plugins/Experimental/MetasoundExperimental/Source/AudioExperimentalRuntime/Private/ChannelAgnostic/ChannelAgnosticTypeUtils.cpp b/Engine/Plugins/Experimental/MetasoundExperimental/Source/AudioExperimentalRuntime/Private/ChannelAgnostic/ChannelAgnosticTypeUtils.cpp
new file mode 100644
--- /dev/null
+++ b/Engine/Plugins/Experimental/MetasoundExperimental/Source/AudioExperimentalRuntime/Private/ChannelAgnostic/ChannelAgnosticTypeUtils.cpp
@@ -0,0 +1,32 @@
+﻿// Copyright Epic Games, Inc. All Rights Reserved.
+
+
+#include "ChannelAgnostic/ChannelAgnosticTypeUtils.h"
+#include "ChannelAgnostic/ChannelAgnosticType.h"
+
+namespace Audio
+{
+	// TODO: MOVE to DSP.
+	static void Interleave(const TArrayView<const float> InMultiMono, const int32 InNumChannels, const TArrayView<float> OutInterleaved)
+	{
+		checkSlow(InNumChannels * InMultiMono.Num() <= OutInterleaved.Num());
+
+		float* Dst = OutInterleaved.GetData();
+		const float* Src = InMultiMono.GetData();
+		const int32 NumFrames = InMultiMono.Num() / InNumChannels;
+					
+		for(int32 Frame = 0; Frame < NumFrames; ++Frame)
+		{
+			for (int32 Channel = 0; Channel < InNumChannels; ++Channel)
+			{
+				*Dst++ = Src[NumFrames * Channel + Frame];
+			}
+		}
+	}
+
+	void FCatUtils::Interleave(const FChannelAgnosticType& In, const TArrayView<float> Out)
+	{
+		Audio::Interleave(In.Buffer.GetView(), In.NumChannels(), Out);
+	}
+
+}
\ No newline at end of file
diff --git a/Engine/Plugins/Experimental/MetasoundExperimental/Source/AudioExperimentalRuntime/Private/TypeFamily/ChannelTypeFamily.cpp b/Engine/Plugins/Experimental/MetasoundExperimental/Source/AudioExperimentalRuntime/Private/TypeFamily/ChannelTypeFamily.cpp
new file mode 100644
--- /dev/null
+++ b/Engine/Plugins/Experimental/MetasoundExperimental/Source/AudioExperimentalRuntime/Private/TypeFamily/ChannelTypeFamily.cpp
@@ -0,0 +1,327 @@
+﻿// Copyright Epic Games, Inc. All Rights Reserved.
+
+#include "TypeFamily/ChannelTypeFamily.h"
+
+#include "CoreGlobals.h"
+#include "Containers/Map.h"
+#include "DSP/ChannelMap.h"
+#include "DSP/MultiMono.h"
+
+namespace Audio
+{
+	namespace ChannelTypeFamilyPrviate
+	{
+		static FString MakePrettyString(const TArray<ESpeakerShortNames>& InEnums)
+		{
+			return FString::JoinBy(InEnums,TEXT(", "), [](ESpeakerShortNames InEnum) -> FString { return LexToString(InEnum); });
+		}
+	}
+	
+	class FChannelRegistryImpl final : public IChannelTypeRegistry 
+	{
+		TMap<FName, const FChannelTypeFamily*> Types;
+		TMap<int32, const FChannelTypeFamily*> ChannelCountLookup;
+	public:
+		UE_NONCOPYABLE(FChannelRegistryImpl);
+		FChannelRegistryImpl() = default;
+		virtual ~FChannelRegistryImpl() override = default;
+	protected:
+		virtual bool RegisterType(const FName& InUniqueName, TRetainedRef<FTypeFamily> InType) override
+		{
+			static const FName ChannelTypeName = TEXT("Cat");
+			const FTypeFamily* Type = FindTypeInternal(ChannelTypeName);
+			if (Type && !InType.Get().IsA(Type))
+			{
+				// Not a cat. (maybe a dog?)
+				return false;	
+			}
+			
+			// TODO: thread safety here.
+			if (Types.Find(InUniqueName) != nullptr)
+			{
+				return false;
+			}
+			
+			// Safe to cast it and add.
+			const FChannelTypeFamily* ChannelType = static_cast<const FChannelTypeFamily*>(&InType.Get());
+
+			// Log.
+			//UE_LOG(LogTemp, Display, TEXT("Registering %s, NumChannels=%d"), *ChannelType->GetName().GetPlainNameString(), ChannelType->NumChannels());
+			
+			// Add to master list of registered types.
+			Types.Add(InUniqueName,ChannelType);
+
+			// If this is a default, then add it to our Channel -> Type lookup table. Assert it's unique.
+			if (const int32 NumChannels = ChannelType->NumChannels(); NumChannels > 0 && ChannelType->IsParentsDefault())
+			{
+				const FChannelTypeFamily*& Found = ChannelCountLookup.FindOrAdd(NumChannels);
+				checkf(Found == nullptr, TEXT("%d Already registered for type %s"), NumChannels, *Found->GetName().ToString());
+				Found = ChannelType;
+			}
+			return true;
+		}	
+		virtual const FTypeFamily* FindTypeInternal(const FName InUniqueName) const override
+		{
+			if (const FChannelTypeFamily* const * Found = Types.Find(InUniqueName))
+			{
+				return *Found;
+			}
+			return nullptr;
+		}
+
+		virtual const FChannelTypeFamily* FindChannelType(const int32 InNumChannels) override
+		{
+			if (const FChannelTypeFamily** Found = ChannelCountLookup.Find(InNumChannels))
+			{
+				check(!(*Found)->IsAbstract());
+				return *Found;
+			}
+
+			// If we don't have a default lookup for this channel count, look up the first thing that matches.
+			using FPair = decltype(Types)::ElementType;
+			TArray<FPair> AllTypes = Types.Array();
+			if (const FPair* Found = AllTypes.FindByPredicate([InNumChannels](const FPair& InPair) -> bool
+			{
+				return InPair.Value->NumChannels() == InNumChannels && !InPair.Value->IsAbstract();
+			}))
+			{
+				return Found->Value;	
+			}
+
+			// Failed to find anything... :(
+			return nullptr;
+		}
+
+		virtual TArray<const FChannelTypeFamily*> GetAllChannelFormats() const override
+		{
+			TArray<const FChannelTypeFamily*> AllFormats;
+			Types.GenerateValueArray(AllFormats);
+			return AllFormats;
+		};
+	};
+
+	IChannelTypeRegistry& GetChannelRegistry()
+	{
+		static FChannelRegistryImpl Registry;
+		return Registry;
+	}
+
+   FChannelTypeFamily::FChannelTypeFamily(
+	   const FName& InUniqueName,
+	   const FName& InFamilyTypeName,
+	   const int32 InNumChannels,
+	   FChannelTypeFamily* InParentType,
+	   const FString& InFriendlyName,
+	   const bool InbIsParentsDefault,
+	   const bool InbIsAbstract)
+	   : Super(InUniqueName, InParentType, InFriendlyName)
+	   , bIsAbstract(InbIsAbstract)
+	   , bIsParentsDefault(InbIsParentsDefault)
+	   , NumChannelsPrivate(InNumChannels)
+	   , FamilyType(InFamilyTypeName)
+	{
+		check(InNumChannels >= 0);
+		check(!InUniqueName.IsNone());
+		if (bIsParentsDefault)
+		{
+			check(InParentType);
+			check(InParentType->DefaultChild == nullptr); 
+			InParentType->DefaultChild = this;
+		}
+	}
+	
+	
+	FDiscreteChannelTypeFamily::FDiscreteChannelTypeFamily(const FName& InUniqueName, FChannelTypeFamily* InParentType, const FString& InFriendlyName, const TArray<ESpeakerShortNames>& InOrder, const bool bIsParentsDefault, const bool bIsAbstract)
+		: Super(InUniqueName, GetFamilyTypeName(), InOrder.Num(), InParentType, InFriendlyName, bIsParentsDefault, bIsAbstract)
+		, Order(InOrder)
+	{
+		checkf(InParentType, TEXT("Type=%s, Has a Null Parent"), *GetName().ToString());
+		auto GetParentNameSafe = [](FChannelTypeFamily* i) -> FString { return i ? i->GetName().ToString() : FString(); };
+		UE_LOG(LogTemp, Display, TEXT("Unique=%s\tNumChannels=%d\tParent=%s\tFriendlyName=%s\tDefault=%s\tOrder=[%s]\tAbstract=%s"),
+			*InUniqueName.ToString(),  InOrder.Num(), *GetParentNameSafe(InParentType), *InFriendlyName, ToCStr(LexToString(bIsParentsDefault)), *ChannelTypeFamilyPrviate::MakePrettyString(Order), ToCStr(LexToString(bIsAbstract)));
+	}
+	
+	FChannelTypeFamily::FTranscoder FDiscreteChannelTypeFamily::GetTranscoder(
+	const FDiscreteChannelTypeFamily& InFromType, const FGetTranscoderParams& InParams) const
+	{
+		// Exact match? We can just memcpy each channel.
+		if (&InFromType == this)
+		{
+			return [NumChannels = InFromType.NumChannels()](TArrayView<const float*> Src, TArrayView<float*> Dst, const int32 NumFrames)
+			{
+				for (int32 i = 0; i < NumChannels; ++i)
+				{
+					FMemory::Memcpy(Dst[i], Src[i], NumFrames * sizeof(float));
+				}
+			};
+		}
+
+		// Make a mix matrix and call a mix/up down.
+		if (TArray<float> Gains; Create2DChannelMap(
+			{
+				.NumInputChannels = InFromType.NumChannels(),
+				.NumOutputChannels = NumChannels(),
+				.Order = EChannelMapOrder::OutputMajorOrder,
+				.MonoUpmixMethod = InParams.MixMethod, 
+				.bIsCenterChannelOnly = InFromType.NumChannels() == 1 && InFromType.HasSpeaker(ESpeakerShortNames::FC) 
+			}, Gains))
+		{
+			return [MixGains = MoveTemp(Gains)](TArrayView<const float*> InSrc, TArrayView<float*> InDst, const int32 NumFrames)
+			{
+				MultiMonoMixUpOrDown(InSrc,InDst,NumFrames,MixGains);
+			};
+		}
+
+		return {};
+	}
+
+	TOptional<FChannelTypeFamily::FChannelName> FDiscreteChannelTypeFamily::GetChannelName(const int32 InChannelIndex) const
+	{
+		check(InChannelIndex >= 0 && InChannelIndex < NumChannels());
+		if (Order.IsValidIndex(InChannelIndex))
+		{
+			const ESpeakerShortNames Speaker = Order[InChannelIndex];
+			const FName SpeakerName = LexToString(Speaker);
+			return
+			{ // Optional.
+					{ 
+						.Name = SpeakerName,
+						.FriendlyName = SpeakerName.ToString() 
+					}
+			};
+		}
+		return {};
+	}
+
+	void RegisterChannelLayouts(IChannelTypeRegistry& Registry)
+	{
+		/**
+		 * Standard layouts.
+		 * This should be defined in .ini ultimately, allowing new, custom formats to be added simply.
+		 **/
+		
+		// Register root type.
+		static FChannelTypeFamily BaseCat(TEXT("Cat"), TEXT("Cat"),  0, nullptr, TEXT("Base Cat"), false, true);
+		Registry.RegisterType(BaseCat.GetName(), BaseCat);
+	
+#define REGISTER_CAT(NAME, PARENT_NAME, FRIENDLY_NAME, LAYOUT, DEFAULT,ABSTRACT)\
+static FDiscreteChannelTypeFamily UE_JOIN(__Type,NAME)(UE_STRINGIZE(NAME), Registry.FindChannel(PARENT_NAME), FRIENDLY_NAME, LAYOUT, DEFAULT, ABSTRACT);\
+Registry.RegisterType(TEXT(UE_STRINGIZE(NAME)),UE_JOIN(__Type,NAME))
+
+		using enum ESpeakerShortNames;
+
+		// Top level abstraction. "Discrete".
+		REGISTER_CAT(Discrete,			TEXT("Cat"),		TEXT("Discrete"),						{},										false,					true);
+	
+		// Mono		(UniqueName)		(Parent Name)			(Friendly Name)										(Channel membership/order.)				(is parents default).  (abstract)
+		REGISTER_CAT(Mono,				TEXT("Discrete"),TEXT("Mono"),							{},										false,					true);
+		REGISTER_CAT(Mono1Dot0,			TEXT("Mono"),	TEXT("Mono (1.0)"),					(TArray{FC}),							true,					false);
+		REGISTER_CAT(Mono1Dot1,			TEXT("Mono"),	TEXT("Mono (1.1)"),					(TArray{FC,LFE}),						false,					false);
+
+		// Stereo																								
+		REGISTER_CAT(Stereo,			TEXT("Discrete"), TEXT("Stereo"),						{},										false,					true);
+		REGISTER_CAT(Stereo2Dot0,		TEXT("Stereo"),  TEXT("Stereo (2.0)"),					(TArray{FL,FR}),						true,					false);
+		REGISTER_CAT(Stereo2Dot1,		TEXT("Stereo"),  TEXT("Stereo (2.1)"),					(TArray{FL,FR,LFE}),					false,					false);
+		REGISTER_CAT(Stereo3Dot0,		TEXT("Stereo"),  TEXT("Stereo (3.0)"),					(TArray{FL,FR,FC}),						false,					false);
+		REGISTER_CAT(Stereo3DOt1,		TEXT("Stereo"),  TEXT("Stereo (3.1)"),					(TArray{FL,FR,FC,LFE}),					false,					false);
+	
+		// Quad 
+		REGISTER_CAT(Quad,				TEXT("Discrete"),	TEXT("Quad"),						{},										false,					true);
+		REGISTER_CAT(Quad4Dot0Back,		TEXT("Quad"),		TEXT("Quad Back Speakers (4.0)"),	(TArray{FL,FR,BL,BR}),					true,					false);
+		REGISTER_CAT(Quad4Dot0Side,		TEXT("Quad"),		TEXT("Quad Side Speakers (4.0)"),	(TArray{FL,FR,SL,SR}),					false,					false);
+		REGISTER_CAT(Quad4Dot1,			TEXT("Quad"),		TEXT("Quad Back Centre LFE (4.1)"),(TArray{FL,FR,BL,BR,LFE}),				false,					false);
+
+		// Surround
+		REGISTER_CAT(Surround,			TEXT("Discrete"),	TEXT("Surround"),					{},										false,					true);
+		
+		REGISTER_CAT(Surround5,			TEXT("Surround"),	TEXT("Surround (5.X)"),			{},										true,					true);
+		REGISTER_CAT(Surround5Dot0,		TEXT("Surround5"),  TEXT("Surround (5.0)"),			(TArray{FR,FR,BL,BR,FC}),				false,					false);
+		REGISTER_CAT(Surround5_1,		TEXT("Surround5"),  TEXT("Surround (5.1)"),			(TArray{FL,FR,BL,BR,FC,LFE}),			true,					false);
+
+		REGISTER_CAT(Surround7,			TEXT("Surround"),	TEXT("Surround (7.X)"),			{},										false,					true);
+		REGISTER_CAT(Surround7Dot0,		TEXT("Surround7"),  TEXT("Surround (7.0)"),			(TArray{FL,FR,SL,SR,FC,BL,BR}),			false,					false);
+		REGISTER_CAT(Surround7Dot1,		TEXT("Surround7"),  TEXT("Surround (7.1)"),			(TArray{FL,FR,SL,SR,FC,BL,BR,LFE}),		true,					false);
+		
+		// Atmos.
+		REGISTER_CAT(Atmos,				TEXT("Surround7"),		TEXT("Dolby Atmos Bed"),		{},													false,		true);
+		REGISTER_CAT(Atmos7Dot0Dot2,	TEXT("Atmos"),			TEXT("Dolby Atmos (7.0.2)"),	(TArray{FL,FR,SL,SR,FC,BL,BR,TFL,TFR}),				false,		false);
+		REGISTER_CAT(Atmos7Dot0Dot4,	TEXT("Atmos"),			TEXT("Dolby Atmos (7.0.4)"),	(TArray{FL,FR,SL,SR,FC,BL,BR,TFL,TFR,TBL,TBR}),		false,		false);
+		REGISTER_CAT(Atmos7Dot1Dot2,	TEXT("Atmos"),			TEXT("Dolby Atmos (7.1.2)"),	(TArray{FL,FR,SL,SR,FC,BL,BR,TFL,TFR,LFE}),			false,		false);
+		REGISTER_CAT(Atmos7Dot1Dot4,	TEXT("Atmos"),			TEXT("Dolby Atmos (7.1.4)"),	(TArray{FL,FR,SL,SR,FC,BL,BR,TFL,TFR,TBL,TBR,LFE}),	true,		false);
+
+		// Channel packs.
+		REGISTER_CAT(ChannelPack,		TEXT("Cat"),				TEXT("Channel Pack"),			{},													false,		true);
+		
+		// Packed mono example.
+		// Might make a new type for this, but for now....
+		static TArray<FDiscreteChannelTypeFamily> MultiMonos = []() 
+		{
+			TArray<FDiscreteChannelTypeFamily> Types;
+			constexpr int32 NumPacks = 6;
+			Types.Reserve(NumPacks);
+			for (int32 i = 0; i < NumPacks; ++i)
+			{
+				constexpr ESpeakerShortNames MonoChannel = FC;
+				TArray<ESpeakerShortNames> Order;
+				Order.Init(ESpeakerShortNames::FC, i+2); 
+				Types.Emplace(
+					*FString::Printf(TEXT("MonoPack%d"), i+2),
+					GetChannelRegistry().FindChannel(TEXT("ChannelPack")),
+					FString::Printf(TEXT("Mono Pack %d Channels"), i+2),
+					Order,	
+					false,
+					false
+				);
+				GetChannelRegistry().RegisterType(Types[i].GetName(),Types[i]);
+			}
+			return Types;
+		}();
+					
+#define REGISTER_AMBISONICS(NAME, PARENT_NAME, FRIENDLY_NAME, ORDER, DEFAULT,ABSTRACT)\
+static FAmbisonicsChannelTypeFamily UE_JOIN(__Type,NAME)(UE_STRINGIZE(NAME), ORDER, Registry.FindChannel(PARENT_NAME), FRIENDLY_NAME, DEFAULT, ABSTRACT);\
+ensure(Registry.RegisterType(UE_STRINGIZE(NAME),UE_JOIN(__Type,NAME)))
+
+		// Ambisonics 
+		REGISTER_AMBISONICS(Ambisonics,				 TEXT("Cat"), 		TEXT("Ambisonics"), 							0, false, true);
+		REGISTER_AMBISONICS(AmbisonicsFirstOrder,	 TEXT("Ambisonics"), TEXT("First Order Ambisonics (4 channels)"), 	1, false, false); 
+		REGISTER_AMBISONICS(AmbisonicsSecondOrder,	 TEXT("Ambisonics"), TEXT("Second Ambisonics (9 channels)"), 		2, false, false); 
+		REGISTER_AMBISONICS(AmbisonicsThirdOrder,	 TEXT("Ambisonics"), TEXT("Third Ambisonics (16 channels)"), 		3, false, false);
+		REGISTER_AMBISONICS(AmbisonicsFourthOrder,	 TEXT("Ambisonics"), TEXT("Fourth Ambisonics (25 chhanels)"), 		4, false, false); 
+		REGISTER_AMBISONICS(AmbisonicsFirthOrder,	 TEXT("Ambisonics"), TEXT("Firth Ambisonics (36 channels)"), 		5, false, false);  
+
+		// Channel packs.
+		// FPackedChannelTypeFamily MonoPack( TEXT("PackedChannel")
+		
+#undef REGISTER_CAT
+#undef REGISTER_AMBISONICS
+	}
+}
+
+const TCHAR* LexToString(const ESpeakerShortNames InSpeaker)
+{
+	// No code gen support here, so role this manually.
+#define CASE_AND_STRING(X) case ESpeakerShortNames::X: return TEXT(#X)
+
+	using enum ESpeakerShortNames;
+	switch (InSpeaker)
+	{
+	CASE_AND_STRING(FL);   // Front Left
+	CASE_AND_STRING(FR);   // Front Right
+	CASE_AND_STRING(FC);   // Front Center
+	CASE_AND_STRING(LFE);  // Low-Frequency Effects (Subwoofer)
+	CASE_AND_STRING(FLC);  // Front Left Center
+	CASE_AND_STRING(FRC);  // Front Right Center
+	CASE_AND_STRING(SL);   // Side Left
+	CASE_AND_STRING(SR);   // Side Right
+	CASE_AND_STRING(BL);   // Back Left
+	CASE_AND_STRING(BR);   // Back Right
+	CASE_AND_STRING(BC);   // Back Center
+	CASE_AND_STRING(TFL);  // Top Front Left
+	CASE_AND_STRING(TFR);  // Top Front Right
+	CASE_AND_STRING(TBL);  // Top Back Left
+	CASE_AND_STRING(TBR);  // Top Back Right
+	default:
+		break;
+	}
+	return nullptr;
+}
\ No newline at end of file
diff --git a/Engine/Plugins/Experimental/MetasoundExperimental/Source/AudioExperimentalRuntime/Public/ChannelAgnostic/ChannelAgnosticType.h b/Engine/Plugins/Experimental/MetasoundExperimental/Source/AudioExperimentalRuntime/Public/ChannelAgnostic/ChannelAgnosticType.h
new file mode 100644
--- /dev/null
+++ b/Engine/Plugins/Experimental/MetasoundExperimental/Source/AudioExperimentalRuntime/Public/ChannelAgnostic/ChannelAgnosticType.h
@@ -0,0 +1,90 @@
+﻿// Copyright Epic Games, Inc. All Rights Reserved.
+
+#pragma once
+
+#include "Containers/ArrayView.h"
+#include "HAL/Platform.h"
+#include "ScratchBuffer.h"
+#include "Templates/RetainedRef.h"
+#include "TypeFamily/ChannelTypeFamily.h"
+
+namespace Audio
+{
+	class FChannelTypeFamily;
+	class FSimpleAllocBase;
+
+	class AUDIOEXPERIMENTALRUNTIME_API FChannelAgnosticType
+	{
+	public:
+		static FSimpleAllocBase& GetDefaultAllocator();
+
+		explicit FChannelAgnosticType(const TRetainedRef<const FChannelTypeFamily> InType, const int32 InNumFrames, FSimpleAllocBase* InAllocator = &GetDefaultAllocator());
+
+		/**
+		 * 
+		 * @param InType 
+		 * @param InNumFrames 
+		 * @param InNumChannels 
+		 * @param InAllocator 
+		 */
+		explicit FChannelAgnosticType(const TRetainedRef<const FChannelTypeFamily> InType, const int32 InNumFrames, const int32 InNumChannels, FSimpleAllocBase* InAllocator = &GetDefaultAllocator());
+	
+		// To be a Variable it *must* be copyable.
+		FChannelAgnosticType(const FChannelAgnosticType& Other) = default;
+		FChannelAgnosticType(FChannelAgnosticType&& Other) = default;
+		FChannelAgnosticType& operator=(FChannelAgnosticType&& Other) = default;
+		FChannelAgnosticType& operator=(const FChannelAgnosticType& Other) = default;
+
+		[[nodiscard]]
+		int32 NumFrames() const { return NumFramesPrivate; }
+
+		[[nodiscard]]
+		int32 NumChannels() const { return NumChannelsPrivate; }
+
+		[[nodiscard]]
+		TArrayView<const float> GetChannel(const int32 InChannelIndex) const
+		{
+			check(InChannelIndex >= 0 && InChannelIndex < NumChannelsPrivate);
+			const TArrayView<const float> AllChannels = Buffer.GetView();
+			return MakeArrayView(AllChannels.GetData() + (NumFramesPrivate * InChannelIndex), NumFramesPrivate);
+		}
+
+		[[nodiscard]]
+		TArrayView<float> GetChannel(const int32 InChannelIndex)
+		{
+			const TArrayView<const float> ConstView = std::as_const(*this).GetChannel(InChannelIndex);
+			return MakeArrayView<float>(const_cast<float*>(ConstView.GetData()), ConstView.Num());
+		}
+
+		bool IsA(const FChannelAgnosticType& InOther) const;
+		bool IsA(const FName& InTypeName) const;
+		FName GetTypeName() const;
+		const FChannelTypeFamily& GetType() const { return *Type; }
+
+		/**
+		 * For fast raw DSP access to the buffer.
+		 * @return View of entire multi-mono buffer.
+		 */
+		[[nodiscard]]
+		TArrayView<float> GetRawMultiMono()
+		{
+			return Buffer.GetView();
+		}
+		/**
+		 * For fast raw DSP access to the buffer.  
+		 * @return View of entire multi-mono buffer.
+		 */
+		[[nodiscard]]
+		TArrayView<const float> GetRawMultiMono() const
+		{
+			return Buffer.GetView();	
+		}
+		
+	private:
+		friend class FCatUtils;
+		TScratchBuffer<float> Buffer;
+		const FChannelTypeFamily* Type = nullptr;	// Keep this as pointer internally for easy of copying.
+		int32 NumFramesPrivate = 0;
+		int32 NumChannelsPrivate = 0;
+	};
+}
diff --git a/Engine/Plugins/Experimental/MetasoundExperimental/Source/AudioExperimentalRuntime/Public/ChannelAgnostic/ChannelAgnosticTypeUtils.h b/Engine/Plugins/Experimental/MetasoundExperimental/Source/AudioExperimentalRuntime/Public/ChannelAgnostic/ChannelAgnosticTypeUtils.h
new file mode 100644
--- /dev/null
+++ b/Engine/Plugins/Experimental/MetasoundExperimental/Source/AudioExperimentalRuntime/Public/ChannelAgnostic/ChannelAgnosticTypeUtils.h
@@ -0,0 +1,17 @@
+﻿// Copyright Epic Games, Inc. All Rights Reserved.
+
+#pragma once
+
+#include "Containers/ArrayView.h"
+
+namespace Audio
+{
+	// Forward decl.
+	class FChannelAgnosticType;
+	
+	class FCatUtils
+	{
+	public:
+		static AUDIOEXPERIMENTALRUNTIME_API void Interleave(const FChannelAgnosticType& In, TArrayView<float> Out);
+	};
+}
diff --git a/Engine/Plugins/Experimental/MetasoundExperimental/Source/AudioExperimentalRuntime/Public/ScratchBuffer.h b/Engine/Plugins/Experimental/MetasoundExperimental/Source/AudioExperimentalRuntime/Public/ScratchBuffer.h
new file mode 100644
--- /dev/null
+++ b/Engine/Plugins/Experimental/MetasoundExperimental/Source/AudioExperimentalRuntime/Public/ScratchBuffer.h
@@ -0,0 +1,126 @@
+// Copyright Epic Games, Inc. All Rights Reserved.
+
+#pragma once
+
+#include "SimpleAlloc/SimpleAllocBase.h"
+
+namespace Audio
+{
+	template <typename T>
+	class AUDIOEXPERIMENTALRUNTIME_API TScratchBuffer
+	{
+	public:
+		using ValueType = T;
+
+		explicit TScratchBuffer(const int32 InNum, FSimpleAllocBase* InAllocator)
+			: Allocator(InAllocator)
+		{
+			DoAllocation(InNum);
+		}
+
+		TScratchBuffer(const TScratchBuffer& Other)
+			: Allocator(Other.Allocator)
+		{
+			*this = Other; // Call copy operator.
+		}
+
+		TScratchBuffer(TScratchBuffer&& Other)
+			: Allocator(Other.Allocator)
+		{
+			*this = MoveTemp(Other); // Call move operator.
+		}
+
+		TScratchBuffer& operator=(const TScratchBuffer& Other)
+		{
+			if (this != &Other)
+			{
+				// Free our state.
+				FreeAllocation();
+				
+				// Make new allocation.
+				Allocator = Other.Allocator;
+				DoAllocation(Other.Allocation.Num());
+
+				// Copy other into new allocated memory 
+				if (Allocation.GetData() && Allocation.Num() > 0)
+				{
+					FMemory::Memcpy(Allocation.GetData(), Other.Allocation.GetData(), Allocation.GetTypeSize() * Allocation.Num());
+				}
+			}
+			return *this;
+		}
+
+
+		// Move operator.
+		TScratchBuffer& operator=(TScratchBuffer&& Other)
+		{
+			if (this == &Other) return *this;
+			Allocator = Other.Allocator;
+			Allocation = MoveTemp(Other.Allocation);
+			LifetimeToken = Other.LifetimeToken;
+
+			// Reset just in case.
+			Other.Reset();
+			return *this;
+		}
+
+		~TScratchBuffer()
+		{
+			FreeAllocation();
+		}
+
+		bool IsValid() const
+		{
+			return Allocator->GetCurrentLifetime() == LifetimeToken;
+		}
+
+		int32 Num() const
+		{
+			check(IsValid());
+			return Allocation.Num();
+		}
+
+		TArrayView<ValueType> GetView()
+		{
+			check(IsValid());
+			return MakeArrayView(Allocation.GetData(), Num());
+		}
+
+		TArrayView<const ValueType> GetView() const
+		{
+			check(IsValid());
+			return MakeArrayView(Allocation.GetData(), Num());
+		}
+
+	private:
+		void Reset()
+		{
+			Allocation = {};
+			LifetimeToken = -1;
+		}
+
+		void DoAllocation(const int32 InNum)
+		{
+			check(InNum > 0);
+			check(Allocation.GetData() == nullptr);
+			Allocation = MakeArrayView<ValueType>(
+				static_cast<ValueType*>(Allocator->Malloc(InNum * sizeof(ValueType))),
+				InNum);
+			check(Allocation.GetData() != nullptr);
+			LifetimeToken = Allocator->GetCurrentLifetime();
+			FMemory::Memzero(Allocation.GetData(), Allocation.GetTypeSize() * Allocation.Num());
+		}
+		void FreeAllocation()
+		{
+			if (Allocation.GetData() != nullptr)
+			{
+				Allocator->Free(Allocation.GetData());
+				Allocation = {};
+			}
+		}
+
+		FSimpleAllocBase* Allocator = nullptr;
+		TArrayView<ValueType> Allocation;
+		int32 LifetimeToken = -1;
+	};
+}
\ No newline at end of file
diff --git a/Engine/Plugins/Experimental/MetasoundExperimental/Source/AudioExperimentalRuntime/Public/SimpleAlloc/SimpleAllocBase.h b/Engine/Plugins/Experimental/MetasoundExperimental/Source/AudioExperimentalRuntime/Public/SimpleAlloc/SimpleAllocBase.h
new file mode 100644
--- /dev/null
+++ b/Engine/Plugins/Experimental/MetasoundExperimental/Source/AudioExperimentalRuntime/Public/SimpleAlloc/SimpleAllocBase.h
@@ -0,0 +1,49 @@
+// Copyright Epic Games, Inc. All Rights Reserved.
+
+#pragma once
+
+#include "HAL/UnrealMemory.h"
+#include "Misc/CoreMiscDefines.h"
+#include "Math/UnrealMathUtility.h"
+
+namespace Audio
+{
+	class FSimpleAllocBase
+	{
+	public:
+		UE_NONCOPYABLE(FSimpleAllocBase);
+		FSimpleAllocBase() = default;
+		virtual ~FSimpleAllocBase() = default;
+	
+		[[nodiscard]] virtual void* Malloc(const SIZE_T, const uint32 InAlignment=DEFAULT_ALIGNMENT) = 0;
+		virtual void Free(void*) {};
+		virtual uint32 GetCurrentLifetime() const { return 0; };
+		virtual void Reset() {};
+	};
+
+	namespace SimpleAllocBasePrivate
+	{
+		// This is what's defined in MemoryBase, so follow it here. 
+		FORCEINLINE static uint32 GetDefaultSizeToAlignment(const uint32 InSize)
+		{
+			check(InSize > 0);
+			return InSize >= 16 ? 16 : 8;
+		};
+
+		// When given an alignment and offset return a new rounded offset that would honor the alignment
+		FORCEINLINE static uint32 RoundUpToAlignment(const uint32 InOffset, const uint32 InAlignment)
+		{
+			check(InAlignment > 0);
+			check(FMath::IsPowerOfTwo(InAlignment));
+			const uint32 Mod = InOffset & (InAlignment-1);
+			const uint32 RoundUp = Mod > 0 ? InAlignment - Mod : 0;
+			return InOffset + RoundUp;
+		}
+
+		FORCEINLINE static bool IsAligned(const void* InPtr, const uint32 InAlignment)
+		{
+			check(FMath::IsPowerOfTwo(InAlignment));
+			return (reinterpret_cast<UPTRINT>(InPtr) & (InAlignment - 1)) == 0;
+		}
+	} //namespace FSimpleAllocBasePrivate
+}
\ No newline at end of file
diff --git a/Engine/Plugins/Experimental/MetasoundExperimental/Source/AudioExperimentalRuntime/Public/SimpleAlloc/SimpleHeapAllocator.h b/Engine/Plugins/Experimental/MetasoundExperimental/Source/AudioExperimentalRuntime/Public/SimpleAlloc/SimpleHeapAllocator.h
new file mode 100644
--- /dev/null
+++ b/Engine/Plugins/Experimental/MetasoundExperimental/Source/AudioExperimentalRuntime/Public/SimpleAlloc/SimpleHeapAllocator.h
@@ -0,0 +1,22 @@
+﻿// Copyright Epic Games, Inc. All Rights Reserved.
+
+#pragma once
+
+#include "SimpleAllocBase.h"
+
+namespace Audio
+{
+	class FSimpleHeapAllocator final : public FSimpleAllocBase
+	{
+	public:
+		FSimpleHeapAllocator() = default;
+		virtual void* Malloc(const SIZE_T InSizeBytes, const uint32 InAlignment=DEFAULT_ALIGNMENT) override
+		{
+			return FMemory::Malloc(InSizeBytes, InAlignment);
+		}
+		virtual void Free(void* InPtr) override
+		{
+			return FMemory::Free(InPtr);
+		}
+	};
+}
diff --git a/Engine/Plugins/Experimental/MetasoundExperimental/Source/AudioExperimentalRuntime/Public/SimpleAlloc/SimpleLinearAllocator.h b/Engine/Plugins/Experimental/MetasoundExperimental/Source/AudioExperimentalRuntime/Public/SimpleAlloc/SimpleLinearAllocator.h
new file mode 100644
--- /dev/null
+++ b/Engine/Plugins/Experimental/MetasoundExperimental/Source/AudioExperimentalRuntime/Public/SimpleAlloc/SimpleLinearAllocator.h
@@ -0,0 +1,64 @@
+﻿// Copyright Epic Games, Inc. All Rights Reserved.
+
+#pragma once
+
+#include "SimpleAllocBase.h"
+#include "Containers/ArrayView.h"
+
+namespace Audio
+{
+	class FSimpleLinearAllocator : public FSimpleAllocBase
+	{
+	public:
+		explicit FSimpleLinearAllocator(const TArrayView<uint8> InPage)
+			: Page(InPage)
+		{
+			check(SimpleAllocBasePrivate::IsAligned(InPage.GetData(),SimpleAllocBasePrivate::GetDefaultSizeToAlignment(InPage.Num())));
+		}
+	
+		virtual void* Malloc(const SIZE_T InSizeBytes, const uint32 InAlignment=DEFAULT_ALIGNMENT) override
+		{
+			const uint32 Alignment = InAlignment == DEFAULT_ALIGNMENT ? SimpleAllocBasePrivate::GetDefaultSizeToAlignment(InSizeBytes) : InAlignment;
+			const uint32 AlignedNewTop = SimpleAllocBasePrivate::RoundUpToAlignment(Top, Alignment);
+		
+			if (AlignedNewTop+InSizeBytes < Page.Num())
+			{
+				Top = AlignedNewTop;
+				void* Ret = Page.GetData() + Top;
+				Top += InSizeBytes;
+				return Ret;
+			}
+			return nullptr;
+		}
+		virtual uint32 GetCurrentLifetime() const override
+		{
+			return CurrentLifetime;
+		}
+		virtual void Reset() override
+		{
+			Top = 0;
+			CurrentLifetime++;
+		}
+	
+	protected:	
+		uint32 Top = 0;					// Top of all allocations, climbs upwards towards end of page.
+		uint32 CurrentLifetime = 0;		// Each reset the lifetime is increased
+		TArrayView<uint8> Page;			// One page for now.
+	};
+
+
+	class FSimpleLinearAllocatorFromHeap final : public FSimpleLinearAllocator
+	{
+	public:
+		explicit FSimpleLinearAllocatorFromHeap(const SIZE_T InPageSize, const uint32 InPageAlignment = DEFAULT_ALIGNMENT, FMalloc& InHeap = *UE::Private::GMalloc)
+			: FSimpleLinearAllocator(MakeArrayView(static_cast<uint8*>(InHeap.Malloc(InPageSize,InPageAlignment)), IntCastChecked<int32>(InPageSize)))
+			, Heap(InHeap)
+		{}
+		virtual ~FSimpleLinearAllocatorFromHeap() override
+		{
+			Heap.Free(Page.GetData());
+		}
+	private:
+		FMalloc& Heap;
+	};
+}
diff --git a/Engine/Plugins/Experimental/MetasoundExperimental/Source/AudioExperimentalRuntime/Public/TypeFamily/ChannelTypeFamily.h b/Engine/Plugins/Experimental/MetasoundExperimental/Source/AudioExperimentalRuntime/Public/TypeFamily/ChannelTypeFamily.h
new file mode 100644
--- /dev/null
+++ b/Engine/Plugins/Experimental/MetasoundExperimental/Source/AudioExperimentalRuntime/Public/TypeFamily/ChannelTypeFamily.h
@@ -0,0 +1,324 @@
+// Copyright Epic Games, Inc. All Rights Reserved.
+
+#pragma once
+
+#include "TypeFamily.h"
+#include "DSP/ChannelMap.h"
+#include "Templates/Function.h"
+#include "Misc/Optional.h"
+
+//#include "ChannelTypeFamily.generated.h"
+
+/**
+ * A Comprehensive Short form Speaker enumeration
+ */
+enum class ESpeakerShortNames : uint32
+{
+	FL,   // Front Left
+	FR,   // Front Right
+	FC,   // Front Center
+	LFE,  // Low-Frequency Effects (Subwoofer)
+	FLC,  // Front Left Center
+	FRC,  // Front Right Center
+	SL,   // Side Left
+	SR,   // Side Right
+	BL,   // Back Left
+	BR,   // Back Right
+	BC,   // Back Center
+	TFL,  // Top Front Left
+	TFR,  // Top Front Right
+	TBL,  // Top Back Left
+	TBR   // Top Back Right
+};
+
+// Pasted from UHT generated.h, but we can't use because of AudioUnitTests.
+#define FOREACH_ENUM_ESPEAKERSHORTNAMES(op) \
+	op(ESpeakerShortNames::FL) \
+	op(ESpeakerShortNames::FR) \
+	op(ESpeakerShortNames::FC) \
+	op(ESpeakerShortNames::LFE) \
+	op(ESpeakerShortNames::FLC) \
+	op(ESpeakerShortNames::FRC) \
+	op(ESpeakerShortNames::SL) \
+	op(ESpeakerShortNames::SR) \
+	op(ESpeakerShortNames::BL) \
+	op(ESpeakerShortNames::BR) \
+	op(ESpeakerShortNames::BC) \
+	op(ESpeakerShortNames::TFL) \
+	op(ESpeakerShortNames::TFR) \
+	op(ESpeakerShortNames::TBL) \
+	op(ESpeakerShortNames::TBR)
+
+const TCHAR* LexToString(const ESpeakerShortNames InSpeaker);
+
+namespace Audio
+{
+	// Forward declare these for visitor.
+	class FDiscreteChannelTypeFamily;
+	class FAmbisonicsChannelTypeFamily;
+	class FPackedChannelTypeFamily;
+		
+	// Base type for all "channel" based types.
+	class FChannelTypeFamily : public FTypeFamily
+	{
+	public:
+		UE_NONCOPYABLE(FChannelTypeFamily);
+		using Super = FTypeFamily;
+
+		/**
+		 * Constructor. 
+		 * @param InUniqueName The unique name that will be used for look up in registry.
+		 * @param InFamilyTypeName The name of the concrete type (c++) that's defining this. (for safe casting). 
+		 * @param InNumChannels Num of channels in this type. (pure categorical, organisational entries will be 0).
+		 * @param InParentType Pointer to this parents type in the tree. Can be null.
+		 * @param InFriendlyName Friendly name to display to the user e.g. "Dolby Stereo (2.0)"
+		 * @param InbIsParentsDefault Marks if this entry is the default child of its parent. 
+		 */
+		explicit FChannelTypeFamily(
+			const FName& InUniqueName,
+			const FName& InFamilyTypeName,
+			const int32 InNumChannels,
+			FChannelTypeFamily* InParentType,
+			const FString& InFriendlyName,
+			const bool InbIsParentsDefault,
+			const bool InbIsAbstract);
+
+		virtual ~FChannelTypeFamily() = default;
+
+		/**
+		 * If this type can be instantiated or is just organisational.
+		 * @return true if abstract, false otherwise
+		 */
+		bool IsAbstract() const
+		{
+			return bIsAbstract;
+		}
+
+		/**
+		 * If this is marked as being the default on its parent
+		 * @return true if true, false otherwise.
+		 */
+		bool IsParentsDefault() const
+		{
+			return bIsParentsDefault;
+		}
+
+		/**
+		 * Returns the default child if one exists. Example. "Stereo" would return "Stereo_2_0"
+		 * @return Pointer to type, null otherwise.
+		 */
+		const FChannelTypeFamily* GetDefaultChild() const { return DefaultChild; }
+
+		/**
+		 * Num of channels  
+		 * @return >= 0
+		 */
+		int32 NumChannels() const { return NumChannelsPrivate; }
+
+		/**
+		 * Transcoder function object.
+		 * Arrays of pointers necessary to make more flexible than just multi-mono and compat with our exisiting libs.
+		 */
+		using FTranscoder = TFunction<void(TArrayView<const float*>, TArrayView<float*>, const int32 NumFrames)>;
+		
+		
+		/**
+		* Parameters to pass to GetTranscoder 
+		* To be expanded.
+		*/
+		struct FGetTranscoderParams
+		{
+			const FChannelTypeFamily& ToType;
+			EChannelMapMonoUpmixMethod MixMethod = EChannelMapMonoUpmixMethod::EqualPower;
+		};
+		/**
+		 * Get the function used for translating/transcoding between this and the "To" type.
+		 * @return The Function Object. (this will be not be set if a function wasn't returned)
+		 */
+		virtual FTranscoder GetTranscoder(const FGetTranscoderParams& InParam) const
+		{
+			return InParam.ToType.GetTranscoder(InParam);
+		}
+
+		/**
+		 * Returns the family name (C++ type) of this ChannelType. This allows us to safely downcast (if necessary).
+		 * @return The name of this type Family. (i.e. Discrete/Ambisoncs). 
+		 */
+		FName GetFamilyName() const { return FamilyType; }
+
+		/**
+		 * Returns the name of the channel. For discrete this would be what speaker its mapped to etc.
+		 * @param InChannelIndex Index of channel >= 0 and < NumChannels.
+		 * @return Optional containing the Channel name and friendly name.
+		 */
+		struct FChannelName
+		{
+			FName Name;
+			FString FriendlyName;
+		};
+		virtual TOptional<FChannelName> GetChannelName(const int32 InChannelIndex) const { return {}; }
+
+	protected:
+		// Both are friends so we can hide the api.
+		friend class FDiscreteChannelTypeFamily;
+		friend class FAmbisonicsChannelTypeFamily;
+		
+		// Visitor like access for transcoder matching.
+		virtual FTranscoder GetTranscoder(const FChannelTypeFamily&, const FGetTranscoderParams&) const  { return {}; }
+		virtual FTranscoder GetTranscoder(const FDiscreteChannelTypeFamily&, const FGetTranscoderParams&) const  { return {}; }
+		virtual FTranscoder GetTranscoder(const FAmbisonicsChannelTypeFamily& , const FGetTranscoderParams&) const  { return {}; }
+		
+	private:
+		bool bIsAbstract = false;
+		bool bIsParentsDefault = false;
+		const FChannelTypeFamily* DefaultChild = nullptr;
+		int32 NumChannelsPrivate = 0;
+		FName FamilyType; 
+	};
+
+	class FDiscreteChannelTypeFamily : public FChannelTypeFamily 
+	{
+	public:
+		UE_NONCOPYABLE(FDiscreteChannelTypeFamily);
+		using Super = FChannelTypeFamily;
+
+		static FName GetFamilyTypeName()
+		{
+			static const FName Name = TEXT("Discrete");
+			return Name;
+		}
+			
+		FDiscreteChannelTypeFamily(const FName& InUniqueName, FChannelTypeFamily* InParentType, const FString& InFriendlyName,
+			const TArray<ESpeakerShortNames>& InOrder, const bool bIsParentsDefault, const bool bIsAbstract);
+
+		/**
+		 * Find the index in the list of Channels for speaker.
+		 * @param InSpeaker 
+		 * @return >= 0 (index) INDEX_NONE (failed to find).
+		 */
+		int32 FindSpeakerIndex(const ESpeakerShortNames InSpeaker) const
+		{
+			return Order.IndexOfByKey(InSpeaker);
+		}
+
+		/**
+		 * Checks if a paticular speaker is present in the Channels for this Format.
+		 * @param InSpeaker 
+		 * @return true (exists), false otherwise.
+		 */
+		bool HasSpeaker(const ESpeakerShortNames InSpeaker) const
+		{
+			return FindSpeakerIndex(InSpeaker) != INDEX_NONE;
+		}
+
+		/**
+		 * Get the function used for translating/transcoding between this and the "To" type.
+		 * @param InToType The Type we are transcoding into.
+		 * @return The Function Object. (this will be not be set if a function wasn't returned)
+		 */	
+
+		virtual TOptional<FChannelName> GetChannelName(const int32 InChannelIndex) const override;
+
+	protected:
+		// Visit "to" Type.
+		virtual FTranscoder GetTranscoder(const FGetTranscoderParams& InParams) const override
+		{
+			return InParams.ToType.GetTranscoder(*this, InParams);
+		}
+		// Visitor accept.
+		virtual FTranscoder GetTranscoder(const FDiscreteChannelTypeFamily&, const FGetTranscoderParams&) const override;		
+		TArray<ESpeakerShortNames> Order;
+	};
+
+	// Ambisonics.
+	class FAmbisonicsChannelTypeFamily : public FChannelTypeFamily 
+	{
+	public:
+		UE_NONCOPYABLE(FAmbisonicsChannelTypeFamily);
+		using Super = FChannelTypeFamily;
+		
+		static FName GetFamilyTypeName()
+		{
+			static const FName Name = TEXT("Ambisonics");
+			return Name;
+		}
+
+		static int32 OrderToNumChannels(const int32 InOrder)
+		{
+			// placeholder.
+			return 0;
+		}
+
+		explicit FAmbisonicsChannelTypeFamily(const FName& InUniqueName, const int32 InOrder, FChannelTypeFamily* InParentType, const FString& InFriendlyName, const bool bIsParentsDefault, const bool InbIsAbstract)
+			: Super(InUniqueName, GetFamilyTypeName(), OrderToNumChannels(InOrder), InParentType, InFriendlyName, bIsParentsDefault, InbIsAbstract)
+			, Order(InOrder)
+		{}
+	private:
+		int32 Order = 0;
+	};
+
+	/*
+	class FPackedChannelTypeFamily final : public FChannelTypeFamily
+	{
+	public:
+		UE_NONCOPYABLE(FPackedChannelTypeFamily);
+
+		using Super = FChannelTypeFamily;
+
+		static FName GetFamilyTypeName()
+		{
+			static const FName Name = TEXT("Packed");
+			return Name;
+		}
+		explicit FPackedChannelTypeFamily(const FName& InUniqueName, FChannelTypeFamily* InParentType, const FString& InFriendlyName, const bool bIsParentsDefault)
+			: Super(InUniqueName, GetFamilyTypeName(), 0, InParentType, InFriendlyName, bIsParentsDefault, false)
+		{}
+	};
+	*/
+	
+	// Channel type registry.
+	class IChannelTypeRegistry : public ITypeFamilyRegistry
+	{
+	public:
+		const FChannelTypeFamily* FindChannel(const FName InName) const
+		{
+			return Find<FChannelTypeFamily>(InName);
+		}
+		FChannelTypeFamily* FindChannel(const FName InName)
+		{
+			return const_cast<FChannelTypeFamily*>(const_cast<const IChannelTypeRegistry*>(this)->FindChannel(InName));
+		}
+
+		/**
+		 * Find the first concrete (non-abstract, > 0 channels) channel. (e.g. Stereo -> 'Stereo_2_0')
+		 * Each Abstract type should have a designated default child.
+		 *
+		 * @param InName Name of format.
+		 * @return Pointer to type (on success), nullptr on failure
+		 */
+		const FChannelTypeFamily* FindConcreteChannel(const FName InName) const
+		{
+			// Walk looking for non-abstract default child. (e.g. Surround->Surround5X->Surround 5.1)
+			const FChannelTypeFamily* Channel = FindChannel(InName);
+			while (Channel && Channel->IsAbstract())
+			{
+				Channel = Channel->GetDefaultChild();
+			}
+			return Channel;
+		}
+
+		// Find a channel type with a specific channel count. (this is best guess)
+		virtual const FChannelTypeFamily* FindChannelType(const int32 InNumChannels) = 0;
+
+		/**
+		 * Returns every registered format as an array.
+		 * @return Array of format types.
+		 */
+		virtual TArray<const FChannelTypeFamily*> GetAllChannelFormats() const = 0;
+	};
+
+	// Loose functions for now.
+	AUDIOEXPERIMENTALRUNTIME_API IChannelTypeRegistry& GetChannelRegistry();
+	void RegisterChannelLayouts(IChannelTypeRegistry& Registry = GetChannelRegistry());
+
+}//namespace Audio
diff --git a/Engine/Plugins/Experimental/MetasoundExperimental/Source/AudioExperimentalRuntime/Public/TypeFamily/TypeFamily.h b/Engine/Plugins/Experimental/MetasoundExperimental/Source/AudioExperimentalRuntime/Public/TypeFamily/TypeFamily.h
new file mode 100644
--- /dev/null
+++ b/Engine/Plugins/Experimental/MetasoundExperimental/Source/AudioExperimentalRuntime/Public/TypeFamily/TypeFamily.h
@@ -0,0 +1,76 @@
+﻿// Copyright Epic Games, Inc. All Rights Reserved.
+
+#pragma once
+
+#include "Templates/RetainedRef.h"
+#include "UObject/NameTypes.h"
+
+namespace Audio
+{
+	class FTypeFamily
+	{
+		UE_NONCOPYABLE(FTypeFamily); 
+		
+		FName Name;
+		const FTypeFamily* Parent = nullptr;
+		FString FriendlyName;
+		
+		template<typename Predicate> const FTypeFamily* FindAncestor(Predicate Pred) const
+		{
+			for (const FTypeFamily* i=this; i; i=i->Parent)
+			{
+				if (::Invoke(Pred,i))
+				{
+					return i;
+				}
+			}
+			return nullptr;
+		}
+	public:
+		FTypeFamily(const FName InName, const FTypeFamily* InParent, const FString& InFriendlyName)
+			: Name(InName)
+			, Parent(InParent)
+			, FriendlyName(InFriendlyName)
+		{}
+
+		bool IsA(const FTypeFamily* Other) const { return nullptr != FindAncestor([Other](const FTypeFamily* i) -> bool { return i==Other; }); }
+		bool IsA(const FName& Other) const { return nullptr != FindAncestor([&](const FTypeFamily* i) -> bool { return i->Name==Other; }); }
+		FName GetName() const { return Name; }
+		FString GetFriendlyName() const { return FriendlyName; }
+	};
+
+	// Simple type registry.
+	class ITypeFamilyRegistry
+	{
+	public:
+		virtual ~ITypeFamilyRegistry() = default;
+		virtual bool RegisterType(const FName& InUniqueName, const TRetainedRef<FTypeFamily> InType) = 0;
+
+		template<typename T> const T* Find(const FName& InUniqueName) const
+		{
+			// TODO: test if this is valid by using FName Comparison.	
+			return static_cast<const T*>(FindTypeInternal(InUniqueName));
+		}
+	protected:
+		virtual const FTypeFamily* FindTypeInternal(const FName InUniqueName) const = 0;
+	};
+
+	// Convenience template so we don't have to specialize each Find. 
+	template<typename T> struct TFamilyRegistry
+	{
+		ITypeFamilyRegistry& FamilyRegistry;
+		explicit TFamilyRegistry(ITypeFamilyRegistry& InFamilyRegistry)
+			: FamilyRegistry(InFamilyRegistry)
+		{}
+		const T* Find(const FName& InName) const
+		{
+			return FamilyRegistry.Find<T>(InName);
+		}
+		const T& FindChecked(const FName& InName) const
+		{
+			const T* Found = Find(InName);
+			check(Found);
+			return *Found;
+		}
+	};
+}
\ No newline at end of file
diff --git a/Engine/Plugins/Experimental/MetasoundExperimental/Source/MetasoundExperimentalEditor/MetasoundExperimentalEditor.build.cs b/Engine/Plugins/Experimental/MetasoundExperimental/Source/MetasoundExperimentalEditor/MetasoundExperimentalEditor.build.cs
new file mode 100644
--- /dev/null
+++ b/Engine/Plugins/Experimental/MetasoundExperimental/Source/MetasoundExperimentalEditor/MetasoundExperimentalEditor.build.cs
@@ -0,0 +1,52 @@
+// Copyright Epic Games, Inc. All Rights Reserved.
+
+using UnrealBuildTool;
+using System.IO; // for Path
+
+public class MetasoundExperimentalEditor : ModuleRules
+{
+	public MetasoundExperimentalEditor(ReadOnlyTargetRules Target) : base(Target)
+	{
+		PCHUsage = ModuleRules.PCHUsageMode.UseExplicitOrSharedPCHs;
+
+		PublicIncludePaths.AddRange(
+			new string[]
+			{
+				// ... add public include paths required here ...
+			}
+		);
+
+		PrivateIncludePaths.AddRange(
+			new string[]
+			{
+				// ... add other private include paths required here ...
+			}
+		);
+
+
+		PublicDependencyModuleNames.AddRange(
+			new string[]
+			{
+				"Core",
+				// ... add other public dependencies that you statically link with here ...
+			}
+		);
+
+		PrivateDependencyModuleNames.AddRange(
+			new string[]
+			{
+				"CoreUObject",
+				"Engine",
+
+				// ... add private dependencies that you statically link with here ...	
+			}
+		);
+
+		DynamicallyLoadedModuleNames.AddRange(
+			new string[]
+			{
+				// ... add any modules that your module loads dynamically here ...
+			}
+		);
+	}
+}
diff --git a/Engine/Plugins/Experimental/MetasoundExperimental/Source/MetasoundExperimentalEditor/Private/MetasoundExperimentalEditorModule.cpp b/Engine/Plugins/Experimental/MetasoundExperimental/Source/MetasoundExperimentalEditor/Private/MetasoundExperimentalEditorModule.cpp
new file mode 100644
--- /dev/null
+++ b/Engine/Plugins/Experimental/MetasoundExperimental/Source/MetasoundExperimentalEditor/Private/MetasoundExperimentalEditorModule.cpp
@@ -0,0 +1,29 @@
+// Copyright Epic Games, Inc. All Rights Reserved.
+
+#define LOCTEXT_NAMESPACE "MetasoundExperimentalEditor"
+
+#include "Modules/ModuleInterface.h"
+#include "Modules/ModuleManager.h"
+
+class FMetasoundExperimentalEditorModule : public IModuleInterface
+{
+public:
+	//~ Begin IModuleInterface API
+	virtual void StartupModule() override;
+	virtual void ShutdownModule() override;
+	//~ End IModuleInterface API
+
+private:
+};
+
+void FMetasoundExperimentalEditorModule::StartupModule()
+{
+}
+
+void FMetasoundExperimentalEditorModule::ShutdownModule()
+{
+}
+
+#undef LOCTEXT_NAMESPACE
+
+IMPLEMENT_MODULE(FMetasoundExperimentalEditorModule, MetasoundExperimentalEditor)
diff --git a/Engine/Plugins/Experimental/MetasoundExperimental/Source/MetasoundExperimentalRuntime/MetasoundExperimentalRuntime.Build.cs b/Engine/Plugins/Experimental/MetasoundExperimental/Source/MetasoundExperimentalRuntime/MetasoundExperimentalRuntime.Build.cs
new file mode 100644
--- /dev/null
+++ b/Engine/Plugins/Experimental/MetasoundExperimental/Source/MetasoundExperimentalRuntime/MetasoundExperimentalRuntime.Build.cs
@@ -0,0 +1,38 @@
+// Copyright Epic Games, Inc. All Rights Reserved.
+
+using UnrealBuildTool;
+
+public class MetasoundExperimentalRuntime : ModuleRules
+{
+	public MetasoundExperimentalRuntime(ReadOnlyTargetRules Target) : base(Target)
+	{
+		PCHUsage = ModuleRules.PCHUsageMode.UseExplicitOrSharedPCHs;
+
+		PublicDependencyModuleNames.AddRange(
+			new string[]
+			{
+				"Core",
+				"MetasoundGraphCore",
+				"MetasoundFrontend",
+				"AudioExperimentalRuntime"
+			}
+		);
+		
+		// When the Editor builds all plugins, it will build with engine. In that case add a dependency on CoreUObject.
+		// If we don't do this we will get linker errors for the Engine expecting to find stuff defined here.
+		if (Target.bCompileAgainstEngine)
+		{
+			PublicDependencyModuleNames.Add("CoreUObject");
+		}
+		
+		PrivateDependencyModuleNames.AddRange(
+			new string[]
+			{
+				"Core",
+				"MetasoundGraphCore",
+				"MetasoundFrontend",
+				"MetasoundStandardNodes"
+			}
+		);
+	}
+}
diff --git a/Engine/Plugins/Experimental/MetasoundExperimental/Source/MetasoundExperimentalRuntime/Private/MetasoundCatBreakNode.cpp b/Engine/Plugins/Experimental/MetasoundExperimental/Source/MetasoundExperimentalRuntime/Private/MetasoundCatBreakNode.cpp
new file mode 100644
--- /dev/null
+++ b/Engine/Plugins/Experimental/MetasoundExperimental/Source/MetasoundExperimentalRuntime/Private/MetasoundCatBreakNode.cpp
@@ -0,0 +1,266 @@
+// Copyright Epic Games, Inc. All Rights Reserved.
+
+#include "MetasoundCatBreakNode.h"
+
+#include "MetasoundChannelAgnosticType.h"
+#include "MetasoundExecutableOperator.h"
+#include "MetasoundFacade.h"
+#include "MetasoundNodeRegistrationMacro.h"
+#include "MetasoundVertex.h"
+#include "DSP/MultiMono.h"
+#include "TypeFamily/ChannelTypeFamily.h"
+
+#define LOCTEXT_NAMESPACE "MetasoundStandardNodes_CatBreakNode"
+
+namespace Metasound
+{
+	class FChannelAgnosticType;
+	struct FBuildOperatorParams;
+	class FNodeFacade;
+
+	namespace CatBreakNodePrivate
+	{
+		METASOUND_PARAM(InputFromCat,			"Input", "CAT to Cast");
+		METASOUND_PARAM(OutputToCat,			"Output", "CAT Result");
+
+		const FLazyName InputBaseName{"In"};
+
+		FName MakeOutputVertexName(const int32 InIndex)
+		{
+			FName Name = InputBaseName;
+			Name.SetNumber(InIndex);
+			return Name;
+		}
+
+		TOutputDataVertex<FAudioBuffer> MakeOutputDataVertex(const int32 InIndex, const FName InName, const FString& InFriendlyName)
+		{
+			const FName Name = MakeOutputVertexName(InIndex);
+#if WITH_EDITOR
+			const FText DisplayName = FText::Format(LOCTEXT("In_DisplayName", "{0}"), FText::FromString(InName.ToString()));
+#else
+			const FText DisplayName = FText::GetEmpty();
+#endif
+			return TOutputDataVertex<FAudioBuffer>{Name, FDataVertexMetadata{FText::FromString(*InFriendlyName), DisplayName}};
+		}
+		
+		FVertexInterface MakeClassInterface(const FName Format)
+		{
+			const Audio::FChannelTypeFamily* Found = Audio::GetChannelRegistry().FindConcreteChannel(Format);
+			if (!Found)
+			{
+				return {};
+			}
+			const int32 NumChannels = Found->NumChannels();
+
+			FInputVertexInterface InputInterface;
+            InputInterface.Add(TInputDataVertex<FChannelAgnosticType>(METASOUND_GET_PARAM_NAME_AND_METADATA(OutputToCat)));
+			
+			FOutputVertexInterface OutputInterface;
+			for (int32 ChannelIndex = 0; ChannelIndex < NumChannels; ++ChannelIndex )
+			{
+				if (TOptional<Audio::FChannelTypeFamily::FChannelName> Name = Found->GetChannelName(ChannelIndex); Name.IsSet())
+				{
+					OutputInterface.Add(MakeOutputDataVertex(ChannelIndex, Name->Name, Name->FriendlyName));
+				}
+				else
+				{
+					OutputInterface.Add(MakeOutputDataVertex(ChannelIndex, TEXT("Output"), FString::FromInt(ChannelIndex) ));	
+				}
+			}
+			return{
+				MoveTemp(InputInterface),
+				MoveTemp(OutputInterface)
+			};
+		}
+		
+		class FCatBreakOperatorData final : public TOperatorData<FCatBreakOperatorData>
+		{
+		public:
+			// The OperatorDataTypeName is used when downcasting an IOperatorData to ensure
+			// that the downcast is valid.
+			static const FLazyName OperatorDataTypeName;
+
+			explicit FCatBreakOperatorData(const FName& InToTypeName)
+				: ToTypeName(InToTypeName)
+			{}
+			const FName& GetToType() const
+			{
+				return ToTypeName;
+			}
+		private:
+			FName ToTypeName;
+		};
+
+		// Linkage.
+		const FLazyName FCatBreakOperatorData::OperatorDataTypeName = TEXT("FCatBreakOperatorData");
+
+		// Helper to create array of multi-mono channel pointers from a CAT.
+		Audio::TStackArrayOfPointers<float> MakeMultiMonoPointersFromBufferArray(const TArray<TDataWriteReference<FAudioBuffer>>& InArrayOfCat)
+		{
+			Audio::TStackArrayOfPointers<float> Result {};
+			for (int32 i = 0; i < InArrayOfCat.Num(); ++i)
+			{
+				Result[i] = InArrayOfCat[i]->GetData();
+			}
+			return Result;
+		}
+	}
+	
+	class FCatBreakOperator final : public TExecutableOperator<FCatBreakOperator>
+	{
+	public:
+		using FTranscoder = Audio::FChannelTypeFamily::FTranscoder;
+		FCatBreakOperator(const FBuildOperatorParams& InParams, FChannelAgnosticTypeReadRef&& InInputCat, TArray<TDataWriteReference<FAudioBuffer>>&& InOutputVerticies, FTranscoder&& InTranscoder, const FName InFormat)
+			: InputCat(MoveTemp(InInputCat))
+			, OutputAudioVertices(MoveTemp(InOutputVerticies))
+			, Settings(InParams.OperatorSettings)
+			, Transcoder(InTranscoder)
+			, Format(InFormat)
+		{}
+		virtual ~FCatBreakOperator() override = default;
+
+		static const FVertexInterface& GetDefaultInterface()
+		{
+			using namespace CatBreakNodePrivate ;
+			auto CreateDefaultInterface = []()-> FVertexInterface
+			{
+				// inputs
+				FInputVertexInterface InputInterface;
+				InputInterface.Add(TInputDataVertex<FChannelAgnosticType>(METASOUND_GET_PARAM_NAME_AND_METADATA(InputFromCat)));
+				
+				// outputs
+				FOutputVertexInterface OutputInterface;
+				
+				return FVertexInterface(MoveTemp(InputInterface), MoveTemp(OutputInterface));
+			}; // end lambda: CreateDefaultInterface()
+
+			static const FVertexInterface DefaultInterface = CreateDefaultInterface();
+			return DefaultInterface;
+		}
+			
+		static TUniquePtr<IOperator> CreateOperator(const FBuildOperatorParams& InParams, FBuildResults& OutResults)
+		{
+			using namespace CatBreakNodePrivate;
+
+			const FCatBreakOperatorData* CatBreakData = CastOperatorData<const FCatBreakOperatorData>(InParams.Node.GetOperatorData().Get());
+			if (!CatBreakData)
+			{
+				// Phil, this fires in the node test.
+				return {};
+			}
+			
+			const FName OutputFormat = CatBreakData->GetToType();
+			const FInputVertexInterfaceData& InputData = InParams.InputData;
+			const Audio::FChannelTypeFamily* ConcreteToType = Audio::GetChannelRegistry().FindConcreteChannel(OutputFormat);
+			if (!ConcreteToType)
+			{
+				// Throw editor error, we don't know the format.
+				return {};
+			}
+
+			// Create the input pin.
+			TDataReadReference<FChannelAgnosticType> InputPin = InputData.GetOrCreateDefaultDataReadReference<FChannelAgnosticType>(
+				METASOUND_GET_PARAM_NAME(OutputToCat), InParams.OperatorSettings);
+
+			// Create output based on format.
+			TArray<TDataWriteReference<FAudioBuffer>> OutputAudioVertices;
+			for (int32 i = 0; i < ConcreteToType->NumChannels(); ++i)
+			{
+				OutputAudioVertices.Emplace(FAudioBufferWriteRef::CreateNew(InParams.OperatorSettings));
+			}
+
+			// Always ask for the transcoder 
+			// In the trivial case where we are the same format, this will be just a memcpy.
+			FTranscoder Transcoder = InputPin->GetType().GetTranscoder(
+				{	
+					.ToType = *ConcreteToType
+				});
+			
+			return MakeUnique<FCatBreakOperator>(
+				InParams,
+				MoveTemp(InputPin),
+				MoveTemp(OutputAudioVertices),
+				MoveTemp(Transcoder),
+				ConcreteToType->GetName()	
+			);
+		}
+
+		virtual void BindInputs(FInputVertexInterfaceData& InOutVertexData) override
+		{
+			using namespace CatBreakNodePrivate;
+			InOutVertexData.BindReadVertex(METASOUND_GET_PARAM_NAME(OutputToCat),InputCat);
+		}
+		virtual void BindOutputs(FOutputVertexInterfaceData& InOutVertexData) override
+		{
+			using namespace CatBreakNodePrivate;
+			for (int32 i = 0; i < OutputAudioVertices.Num(); i++)
+			{
+				InOutVertexData.BindWriteVertex(MakeOutputVertexName(i), OutputAudioVertices[i]);
+			}
+		}
+
+		void Reset(const FResetParams& InParams)
+		{
+			Execute();
+		}
+
+		void Execute()
+		{
+			using namespace Audio;
+			using namespace CatBreakNodePrivate;
+			if (Transcoder)
+			{
+				TStackArrayOfPointers<const float> Src = MakeMultiMonoPointersFromView(InputCat->GetRawMultiMono(), Settings.GetNumFramesPerBlock(), InputCat->NumChannels());
+				TStackArrayOfPointers<float> Dst = MakeMultiMonoPointersFromBufferArray(OutputAudioVertices);
+				Transcoder(
+					MakeArrayView(Src.GetData(),InputCat->NumChannels()),
+					MakeArrayView(Dst.GetData(), OutputAudioVertices.Num()),
+					Settings.GetNumFramesPerBlock()
+					);
+			}
+		}
+		
+		static FNodeClassMetadata GetNodeInfo()
+		{
+			using namespace CatBreakNodePrivate;
+			return FNodeClassMetadata
+			{
+				FNodeClassName{ "Experimental", "CatBreakOperator", "" },
+				1, // Major version
+				0, // Minor version
+				LOCTEXT("CatBreakNodeName", "CAT Break Node"),	
+				LOCTEXT("CatBreakNodeNameDescription", "A Node that Breaks CATs"),	
+				TEXT("UE"), // Author
+				LOCTEXT("ExampleConfigurablePromptIfMissing", "Enable the MetaSoundExperimental Plugin"), // Prompt if missing
+				GetDefaultInterface(),
+				{}
+			};
+		}
+	private:
+		FChannelAgnosticTypeReadRef InputCat;
+		TArray<TDataWriteReference<FAudioBuffer>> OutputAudioVertices;
+		FOperatorSettings Settings;
+		FTranscoder Transcoder;
+		FName Format;
+	}; // class FCatCastingOperator
+		
+	using FCatBreakNode = TNodeFacade<FCatBreakOperator>;
+
+	// register node config.	
+	METASOUND_REGISTER_NODE_AND_CONFIGURATION(FCatBreakNode, FMetaSoundCatBreakNodeConfiguration);
+} // namespace Metasound
+
+
+TInstancedStruct<FMetasoundFrontendClassInterface> FMetaSoundCatBreakNodeConfiguration::OverrideDefaultInterface(
+	const FMetasoundFrontendClass& InNodeClass) const
+{
+	using namespace Metasound::CatBreakNodePrivate;
+	return  TInstancedStruct<FMetasoundFrontendClassInterface>::Make(FMetasoundFrontendClassInterface::GenerateClassInterface(MakeClassInterface(Format)));
+}
+
+TSharedPtr<const Metasound::IOperatorData> FMetaSoundCatBreakNodeConfiguration::GetOperatorData() const
+{
+	return MakeShared<Metasound::CatBreakNodePrivate::FCatBreakOperatorData>(Format);
+}
+
+#undef LOCTEXT_NAMESPACE // 
diff --git a/Engine/Plugins/Experimental/MetasoundExperimental/Source/MetasoundExperimentalRuntime/Private/MetasoundCatBreakNode.h b/Engine/Plugins/Experimental/MetasoundExperimental/Source/MetasoundExperimentalRuntime/Private/MetasoundCatBreakNode.h
new file mode 100644
--- /dev/null
+++ b/Engine/Plugins/Experimental/MetasoundExperimental/Source/MetasoundExperimentalRuntime/Private/MetasoundCatBreakNode.h
@@ -0,0 +1,26 @@
+﻿// Copyright Epic Games, Inc. All Rights Reserved.
+
+#pragma once
+
+#include "MetasoundFrontendDocument.h"
+#include "Templates/SharedPointer.h"
+
+#include "MetasoundCatBreakNode.generated.h"
+
+USTRUCT()
+struct FMetaSoundCatBreakNodeConfiguration : public FMetaSoundFrontendNodeConfiguration
+{
+	GENERATED_BODY()
+
+	FMetaSoundCatBreakNodeConfiguration() = default;
+
+	UPROPERTY(EditAnywhere, Category = General, meta = (GetOptions="MetasoundCatCastingOptionsHelper.GetCastingOptions"))
+	FName Format = TEXT("Stereo");
+
+	/* Get the current interface for the class based upon the node extension */
+	virtual TInstancedStruct<FMetasoundFrontendClassInterface> OverrideDefaultInterface(const FMetasoundFrontendClass& InNodeClass) const override;
+
+	/* Pass data down to the operator. */
+	virtual TSharedPtr<const Metasound::IOperatorData> GetOperatorData() const override;
+};
+
diff --git a/Engine/Plugins/Experimental/MetasoundExperimental/Source/MetasoundExperimentalRuntime/Private/MetasoundCatCastingNode.cpp b/Engine/Plugins/Experimental/MetasoundExperimental/Source/MetasoundExperimentalRuntime/Private/MetasoundCatCastingNode.cpp
new file mode 100644
--- /dev/null
+++ b/Engine/Plugins/Experimental/MetasoundExperimental/Source/MetasoundExperimentalRuntime/Private/MetasoundCatCastingNode.cpp
@@ -0,0 +1,190 @@
+// Copyright Epic Games, Inc. All Rights Reserved.
+
+#include "MetasoundCatCastingNode.h"
+
+#include "MetasoundChannelAgnosticType.h"
+#include "MetasoundExecutableOperator.h"
+#include "MetasoundFacade.h"
+#include "MetasoundNodeRegistrationMacro.h"
+#include "MetasoundVertex.h"
+#include "TypeFamily/ChannelTypeFamily.h"
+
+#define LOCTEXT_NAMESPACE "MetasoundStandardNodes_CatCastingNode"
+
+namespace Metasound
+{
+	class FChannelAgnosticType;
+	struct FBuildOperatorParams;
+	class FNodeFacade;
+
+	namespace CatCastingPrivate
+	{
+		METASOUND_PARAM(InputFromCat,			"Input", "CAT to Cast");
+		METASOUND_PARAM(OutputToCat,			"Output", "CAT Result");
+		
+		class FCatCastingOperatorData final : public TOperatorData<FCatCastingOperatorData>
+		{
+		public:
+			// The OperatorDataTypeName is used when downcasting an IOperatorData to ensure
+			// that the downcast is valid.
+			static const FLazyName OperatorDataTypeName;
+
+			explicit FCatCastingOperatorData(const FName& InToTypeName)
+				: ToTypeName(InToTypeName)
+			{}
+			const FName& GetToType() const
+			{
+				return ToTypeName;
+			}
+		private:
+			FName ToTypeName;
+		};
+
+		// Linkage.
+		const FLazyName FCatCastingOperatorData::OperatorDataTypeName = TEXT("FCatCastingOperatorData");
+	}
+	
+	class FCatCastingOperator final : public TExecutableOperator<FCatCastingOperator>
+	{
+	public:
+		FCatCastingOperator(const FBuildOperatorParams& InParams, FChannelAgnosticTypeReadRef&& InInputCat,const FName InToName)
+			: InputFrom(MoveTemp(InInputCat))
+			, ToFormatName(InToName)
+			, OutputCastResult(FChannelAgnosticTypeWriteRef::CreateNew(InParams.OperatorSettings, InToName))
+			, Settings(InParams.OperatorSettings)
+		{
+		}
+		virtual ~FCatCastingOperator() override = default;
+
+		static const FVertexInterface& GetDefaultInterface()
+		{
+			using namespace CatCastingPrivate;
+			auto CreateDefaultInterface = []()-> FVertexInterface
+			{
+				// inputs
+				FInputVertexInterface InputInterface;
+				InputInterface.Add(TInputDataVertex<FChannelAgnosticType>(METASOUND_GET_PARAM_NAME_AND_METADATA(InputFromCat)));
+				
+				// outputs
+				FOutputVertexInterface OutputInterface;
+				OutputInterface.Add(TOutputDataVertex<FChannelAgnosticType>(METASOUND_GET_PARAM_NAME_AND_METADATA(OutputToCat)));
+				
+				return FVertexInterface(InputInterface, OutputInterface);
+			}; // end lambda: CreateDefaultInterface()
+
+			static const FVertexInterface DefaultInterface = CreateDefaultInterface();
+			return DefaultInterface;
+		}
+			
+		static TUniquePtr<IOperator> CreateOperator(const FBuildOperatorParams& InParams, FBuildResults& OutResults)
+		{
+			using namespace CatCastingPrivate;
+
+			const FCatCastingOperatorData* CatTestingConfigData = CastOperatorData<const FCatCastingOperatorData>(InParams.Node.GetOperatorData().Get());
+			check(CatTestingConfigData);
+			const FName CastToName = CatTestingConfigData->GetToType();
+			
+			const FInputVertexInterfaceData& InputData = InParams.InputData;
+			const Audio::FChannelTypeFamily* ConcreteToType = Audio::GetChannelRegistry().FindConcreteChannel(CastToName);
+
+			TDataReadReference<FChannelAgnosticType> InputCat = InputData.GetOrCreateDefaultDataReadReference<FChannelAgnosticType>(
+				METASOUND_GET_PARAM_NAME(InputFromCat), InParams.OperatorSettings);
+
+			// Make sure the cast is to something sane, otherwise use the inputs type... 
+			const FName CastToNameSane = ConcreteToType ? ConcreteToType->GetName() : InputCat->GetTypeName();
+			
+			return MakeUnique<FCatCastingOperator>(
+				InParams,
+				MoveTemp(InputCat),
+				CastToNameSane
+			);
+		}
+
+		virtual void BindInputs(FInputVertexInterfaceData& InOutVertexData) override
+		{
+			using namespace CatCastingPrivate;
+			InOutVertexData.BindReadVertex(METASOUND_GET_PARAM_NAME(InputFromCat) , InputFrom);
+		}
+		
+		virtual void BindOutputs(FOutputVertexInterfaceData& InOutVertexData) override
+		{
+			using namespace CatCastingPrivate;
+			InOutVertexData.BindWriteVertex(METASOUND_GET_PARAM_NAME(OutputToCat),OutputCastResult);
+
+			// Create transcoder 
+			Transcoder = InputFrom->GetType().GetTranscoder(
+				{ .ToType = OutputCastResult->GetType() }
+			);
+		}
+
+		void Reset(const FResetParams& InParams)
+		{
+			Execute();
+		}
+
+		void Execute()
+		{
+			if (Transcoder)
+			{
+				auto Src = Audio::MakeMultiMonoPointersFromView(InputFrom->GetRawMultiMono(), Settings.GetNumFramesPerBlock(), InputFrom->NumChannels());
+				auto Dst = Audio::MakeMultiMonoPointersFromView(OutputCastResult->GetRawMultiMono(), Settings.GetNumFramesPerBlock(), OutputCastResult->NumChannels());				
+				Transcoder(
+				 	MakeArrayView(Src.GetData(),InputFrom->NumChannels()),
+				 	MakeArrayView(Dst.GetData(),OutputCastResult->NumChannels()),
+				 	InputFrom->NumFrames()
+				 	);
+			}
+		}
+		static FNodeClassMetadata GetNodeInfo()
+		{
+			using namespace CatCastingPrivate;
+			return FNodeClassMetadata
+			{
+				FNodeClassName{ "Experimental", "CatCastingOperator", "" },
+				1, // Major version
+				0, // Minor version
+				LOCTEXT("CatCastingNodeName", "CAT Casting Node"),	
+				LOCTEXT("CatCastingNodeNameDescription", "A Node that allows Casting to CATs"),	
+				TEXT("UE"), // Author
+				LOCTEXT("ExampleConfigurablePromptIfMissing", "Enable the MetaSoundExperimental Plugin"), // Prompt if missing
+				GetDefaultInterface(),
+				{}
+			};
+		}
+	private:
+		using FTranscoder = Audio::FChannelTypeFamily::FTranscoder;
+		FChannelAgnosticTypeReadRef InputFrom;
+		FName ToFormatName;
+		FChannelAgnosticTypeWriteRef OutputCastResult;
+		FOperatorSettings Settings;
+		FTranscoder Transcoder;
+	}; // class FCatCastingOperator
+		
+	using FCatCastingNode = TNodeFacade<FCatCastingOperator>;
+
+	// register node config.	
+	METASOUND_REGISTER_NODE_AND_CONFIGURATION(FCatCastingNode, FMetaSoundCatCastingNodeConfiguration);
+} // namespace Metasound
+
+
+TArray<FPropertyTextFName> UMetasoundCatCastingOptionsHelper::GetCastingOptions()
+{
+	const TArray<const Audio::FChannelTypeFamily*> AllFormats = Audio::GetChannelRegistry().GetAllChannelFormats();
+	TArray<FPropertyTextFName> FormatsOptions;
+	Algo::Transform(AllFormats, FormatsOptions, [](const Audio::FChannelTypeFamily* i) -> FPropertyTextFName
+		{ return {.ValueString = i->GetName(), .DisplayName = FText::FromString(i->GetFriendlyName()) }; });
+	return FormatsOptions;
+}
+
+TInstancedStruct<FMetasoundFrontendClassInterface> FMetaSoundCatCastingNodeConfiguration::OverrideDefaultInterface(
+	const FMetasoundFrontendClass& InNodeClass) const
+{
+	return  TInstancedStruct<FMetasoundFrontendClassInterface>::Make(FMetasoundFrontendClassInterface::GenerateClassInterface(Metasound::FCatCastingOperator::GetDefaultInterface()));
+}
+
+TSharedPtr<const Metasound::IOperatorData> FMetaSoundCatCastingNodeConfiguration::GetOperatorData() const
+{
+	return MakeShared<Metasound::CatCastingPrivate::FCatCastingOperatorData>(ToType);
+}
+
+#undef LOCTEXT_NAMESPACE // 
diff --git a/Engine/Plugins/Experimental/MetasoundExperimental/Source/MetasoundExperimentalRuntime/Private/MetasoundCatCastingNode.h b/Engine/Plugins/Experimental/MetasoundExperimental/Source/MetasoundExperimentalRuntime/Private/MetasoundCatCastingNode.h
new file mode 100644
--- /dev/null
+++ b/Engine/Plugins/Experimental/MetasoundExperimental/Source/MetasoundExperimentalRuntime/Private/MetasoundCatCastingNode.h
@@ -0,0 +1,38 @@
+﻿// Copyright Epic Games, Inc. All Rights Reserved.
+
+#pragma once
+
+#include "MetasoundFrontendDocument.h"
+#include "Templates/SharedPointer.h"
+#include "UObject/PropertyText.h"
+
+#include "MetasoundCatCastingNode.generated.h"
+
+UCLASS()
+class UMetasoundCatCastingOptionsHelper : public UObject
+{
+	GENERATED_BODY()
+
+	UFUNCTION()
+	static TArray<FPropertyTextFName> GetCastingOptions();
+};
+
+USTRUCT()
+struct FMetaSoundCatCastingNodeConfiguration : public FMetaSoundFrontendNodeConfiguration
+{
+	GENERATED_BODY()
+
+	FMetaSoundCatCastingNodeConfiguration() =  default;
+
+	UPROPERTY(EditAnywhere, Category = General, meta = (GetOptions="MetasoundCatCastingOptionsHelper.GetCastingOptions"))
+	FName ToType = TEXT("Mono");
+	
+	UPROPERTY(EditAnywhere, Category = General, meta = (ClampMin = "1", ClampMax = "1000"))
+	uint32 NumOutputs = 1;
+
+	/* Get the current interface for the class based upon the node extension */
+	virtual TInstancedStruct<FMetasoundFrontendClassInterface> OverrideDefaultInterface(const FMetasoundFrontendClass& InNodeClass) const override;
+
+	/* Pass data down to the operator. */
+	virtual TSharedPtr<const Metasound::IOperatorData> GetOperatorData() const override;
+};
diff --git a/Engine/Plugins/Experimental/MetasoundExperimental/Source/MetasoundExperimentalRuntime/Private/MetasoundCatMakeNode.cpp b/Engine/Plugins/Experimental/MetasoundExperimental/Source/MetasoundExperimentalRuntime/Private/MetasoundCatMakeNode.cpp
new file mode 100644
--- /dev/null
+++ b/Engine/Plugins/Experimental/MetasoundExperimental/Source/MetasoundExperimentalRuntime/Private/MetasoundCatMakeNode.cpp
@@ -0,0 +1,222 @@
+// Copyright Epic Games, Inc. All Rights Reserved.
+
+#include "MetasoundCatMakeNode.h"
+
+#include "MetasoundChannelAgnosticType.h"
+#include "MetasoundExecutableOperator.h"
+#include "MetasoundFacade.h"
+#include "MetasoundNodeRegistrationMacro.h"
+#include "MetasoundVertex.h"
+#include "TypeFamily/ChannelTypeFamily.h"
+
+#define LOCTEXT_NAMESPACE "MetasoundStandardNodes_CatMakeNode"
+
+namespace Metasound
+{
+	class FChannelAgnosticType;
+	struct FBuildOperatorParams;
+	class FNodeFacade;
+
+	namespace CatMakeNodePrivate
+	{
+		METASOUND_PARAM(InputFromCat,			"Input", "CAT to Cast");
+		METASOUND_PARAM(OutputToCat,			"Output", "CAT Result");
+
+		const FLazyName InputBaseName{"In"};
+
+		FName MakeInputVertexName(const int32 InIndex)
+		{
+			FName Name = InputBaseName;
+			Name.SetNumber(InIndex);
+			return Name;
+		}
+
+		FInputDataVertex MakeInputDataVertex(const int32 InIndex, const FName InName, const FString& InFriendlyName)
+		{
+			const FName InputName = MakeInputVertexName(InIndex);
+#if WITH_EDITOR
+			const FText InputDisplayName = FText::Format(LOCTEXT("In_DisplayName", "{0}"), FText::FromString(InName.ToString()));
+#else
+			const FText InputDisplayName = FText::GetEmpty();
+#endif
+			return TInputDataVertex<FAudioBuffer>{InputName, FDataVertexMetadata{FText::FromString(*InFriendlyName), InputDisplayName}};
+		}
+		
+		
+		FVertexInterface MakeClassInterface(const FName Format)
+		{
+			const Audio::FChannelTypeFamily* Found = Audio::GetChannelRegistry().FindConcreteChannel(Format);
+			if (!Found)
+			{
+				return {};
+			}
+			const int32 NumChannels = Found->NumChannels();
+
+			FInputVertexInterface InputInterface;
+			for (int32 ChannelIndex = 0; ChannelIndex < NumChannels; ++ChannelIndex )
+			{
+				if (TOptional<Audio::FChannelTypeFamily::FChannelName> Name = Found->GetChannelName(ChannelIndex); Name.IsSet())
+				{
+					InputInterface.Add(MakeInputDataVertex(ChannelIndex, Name->Name, Name->FriendlyName));
+				}
+				else
+				{
+					InputInterface.Add(MakeInputDataVertex(ChannelIndex, TEXT("Input"), FString::FromInt(ChannelIndex) ));	
+				}
+			}
+			FOutputVertexInterface OutputInterface;
+            OutputInterface.Add(TOutputDataVertex<FChannelAgnosticType>(METASOUND_GET_PARAM_NAME_AND_METADATA(OutputToCat)));
+			return{
+				MoveTemp(InputInterface),
+				MoveTemp(OutputInterface)
+			};
+		}
+		
+		class FCatMakeOperatorData final : public TOperatorData<FCatMakeOperatorData>
+		{
+		public:
+			// The OperatorDataTypeName is used when downcasting an IOperatorData to ensure
+			// that the downcast is valid.
+			static const FLazyName OperatorDataTypeName;
+
+			explicit FCatMakeOperatorData(const FName& InToTypeName)
+				: ToTypeName(InToTypeName)
+			{}
+			const FName& GetToType() const
+			{
+				return ToTypeName;
+			}
+		private:
+			FName ToTypeName;
+		};
+
+		// Linkage.
+		const FLazyName FCatMakeOperatorData::OperatorDataTypeName = TEXT("FCatMakeOperatorData");
+	}
+	
+	class FCatMakeOperator final : public TExecutableOperator<FCatMakeOperator>
+	{
+	public:
+		FCatMakeOperator(const FBuildOperatorParams& InParams, TArray<TDataReadReference<FAudioBuffer>>&& InInputs, FChannelAgnosticTypeWriteRef&& InOutputCat)
+			: InputAudioVertices(MoveTemp(InInputs))
+			, OutputCat(MoveTemp(InOutputCat))
+			, Settings(InParams.OperatorSettings)
+		{
+		}
+		virtual ~FCatMakeOperator() override = default;
+
+		static const FVertexInterface& GetDefaultInterface()
+		{
+			using namespace CatMakeNodePrivate ;
+			auto CreateDefaultInterface = []()-> FVertexInterface
+			{
+				// outputs
+				FOutputVertexInterface OutputInterface;
+				OutputInterface.Add(TOutputDataVertex<FChannelAgnosticType>(METASOUND_GET_PARAM_NAME_AND_METADATA(OutputToCat)));
+				
+				return FVertexInterface({}, OutputInterface);
+			}; // end lambda: CreateDefaultInterface()
+
+			static const FVertexInterface DefaultInterface = CreateDefaultInterface();
+			return DefaultInterface;
+		}
+			
+		static TUniquePtr<IOperator> CreateOperator(const FBuildOperatorParams& InParams, FBuildResults& OutResults)
+		{
+			using namespace CatMakeNodePrivate;
+
+			const FCatMakeOperatorData* CatMakeData = CastOperatorData<const FCatMakeOperatorData>(InParams.Node.GetOperatorData().Get());
+			check(CatMakeData);
+			const FName OutputFormat = CatMakeData->GetToType();
+			const FInputVertexInterfaceData& InputData = InParams.InputData;
+			const Audio::FChannelTypeFamily* ConcreteToType = Audio::GetChannelRegistry().FindConcreteChannel(OutputFormat);
+
+			// Set the number of inputs pins to match the output format.
+			TArray<TDataReadReference<FAudioBuffer>> InputAudioVertices;
+			for (int32 i = 0; i < ConcreteToType->NumChannels(); ++i)
+			{
+				InputAudioVertices.Emplace(InputData.GetOrCreateDefaultDataReadReference<FAudioBuffer>(MakeInputVertexName(i), InParams.OperatorSettings));
+			}
+						
+			// Make output CAT to match our configs settings. (use concrete form).
+			FChannelAgnosticTypeWriteRef OutputPin = FChannelAgnosticTypeWriteRef::CreateNew(InParams.OperatorSettings, ConcreteToType->GetName());
+			
+			return MakeUnique<FCatMakeOperator>(
+				InParams,
+				MoveTemp(InputAudioVertices),
+				MoveTemp(OutputPin)
+			);
+		}
+
+		virtual void BindInputs(FInputVertexInterfaceData& InOutVertexData) override
+		{
+			using namespace CatMakeNodePrivate;
+			for (int32 i = 0; i < InputAudioVertices.Num(); i++)
+			{
+				InOutVertexData.BindReadVertex(MakeInputVertexName(i), InputAudioVertices[i]);
+			}
+		}
+		virtual void BindOutputs(FOutputVertexInterfaceData& InOutVertexData) override
+		{
+			using namespace CatMakeNodePrivate;
+			InOutVertexData.BindWriteVertex(METASOUND_GET_PARAM_NAME(OutputToCat),OutputCat);
+		}
+
+		void Reset(const FResetParams& InParams)
+		{
+			Execute();
+		}
+
+		void Execute()
+		{
+			// Copy each input channel into the CAT.
+			for (int32 Index = 0; Index < InputAudioVertices.Num(); Index++)
+			{
+				const FAudioBuffer& Src = *InputAudioVertices[Index]; 
+				TArrayView<float> Dst = OutputCat->GetChannel(Index);
+				checkSlow(Src.Num() == Dst.Num());
+				FMemory::Memcpy(Dst.GetData(), Src.GetData(), Dst.Num() * sizeof(float));
+			}
+		}
+		static FNodeClassMetadata GetNodeInfo()
+		{
+			using namespace CatMakeNodePrivate;
+			return FNodeClassMetadata
+			{
+				FNodeClassName{ "Experimental", "CatMakeOperator", "" },
+				1, // Major version
+				0, // Minor version
+				LOCTEXT("CatMakeNodeName", "CAT Make Node"),	
+				LOCTEXT("CatMakeNodeNameDescription", "A Node that builds CATs"),	
+				TEXT("UE"), // Author
+				LOCTEXT("ExampleConfigurablePromptIfMissing", "Enable the MetaSoundExperimental Plugin"), // Prompt if missing
+				GetDefaultInterface(),
+				{}
+			};
+		}
+	private:
+		TArray<TDataReadReference<FAudioBuffer>> InputAudioVertices;
+		FChannelAgnosticTypeWriteRef OutputCat;
+		FOperatorSettings Settings;
+	}; // class FCatCastingOperator
+		
+	using FCatMakeNode = TNodeFacade<FCatMakeOperator>;
+
+	// register node config.	
+	METASOUND_REGISTER_NODE_AND_CONFIGURATION(FCatMakeNode, FMetaSoundCatMakeNodeConfiguration);
+} // namespace Metasound
+
+
+TInstancedStruct<FMetasoundFrontendClassInterface> FMetaSoundCatMakeNodeConfiguration::OverrideDefaultInterface(
+	const FMetasoundFrontendClass& InNodeClass) const
+{
+	using namespace Metasound::CatMakeNodePrivate;
+	return  TInstancedStruct<FMetasoundFrontendClassInterface>::Make(FMetasoundFrontendClassInterface::GenerateClassInterface(MakeClassInterface(Format)));
+}
+
+TSharedPtr<const Metasound::IOperatorData> FMetaSoundCatMakeNodeConfiguration::GetOperatorData() const
+{
+	return MakeShared<Metasound::CatMakeNodePrivate::FCatMakeOperatorData>(Format);
+}
+
+#undef LOCTEXT_NAMESPACE // 
diff --git a/Engine/Plugins/Experimental/MetasoundExperimental/Source/MetasoundExperimentalRuntime/Private/MetasoundCatMakeNode.h b/Engine/Plugins/Experimental/MetasoundExperimental/Source/MetasoundExperimentalRuntime/Private/MetasoundCatMakeNode.h
new file mode 100644
--- /dev/null
+++ b/Engine/Plugins/Experimental/MetasoundExperimental/Source/MetasoundExperimentalRuntime/Private/MetasoundCatMakeNode.h
@@ -0,0 +1,26 @@
+﻿// Copyright Epic Games, Inc. All Rights Reserved.
+
+#pragma once
+
+#include "MetasoundFrontendDocument.h"
+#include "Templates/SharedPointer.h"
+
+#include "MetasoundCatMakeNode.generated.h"
+
+USTRUCT()
+struct FMetaSoundCatMakeNodeConfiguration : public FMetaSoundFrontendNodeConfiguration
+{
+	GENERATED_BODY()
+
+	FMetaSoundCatMakeNodeConfiguration() = default;
+
+	UPROPERTY(EditAnywhere, Category = General, meta = (GetOptions="MetasoundCatCastingOptionsHelper.GetCastingOptions"))
+	FName Format = TEXT("Stereo");
+
+	/* Get the current interface for the class based upon the node extension */
+	virtual TInstancedStruct<FMetasoundFrontendClassInterface> OverrideDefaultInterface(const FMetasoundFrontendClass& InNodeClass) const override;
+
+	/* Pass data down to the operator. */
+	virtual TSharedPtr<const Metasound::IOperatorData> GetOperatorData() const override;
+};
+
diff --git a/Engine/Plugins/Experimental/MetasoundExperimental/Source/MetasoundExperimentalRuntime/Private/MetasoundCatMixerNode.cpp b/Engine/Plugins/Experimental/MetasoundExperimental/Source/MetasoundExperimentalRuntime/Private/MetasoundCatMixerNode.cpp
new file mode 100644
--- /dev/null
+++ b/Engine/Plugins/Experimental/MetasoundExperimental/Source/MetasoundExperimentalRuntime/Private/MetasoundCatMixerNode.cpp
@@ -0,0 +1,160 @@
+// Copyright Epic Games, Inc. All Rights Reserved.
+
+#include "MetasoundCatMixerNode.h"
+
+#include "MetasoundChannelAgnosticType.h"
+#include "MetasoundExecutableOperator.h"
+#include "MetasoundFacade.h"
+#include "MetasoundNodeRegistrationMacro.h"
+#include "MetasoundPrimitives.h"
+#include "MetasoundStandardNodesCategories.h"
+#include "MetasoundVertex.h"
+#include "TypeFamily/ChannelTypeFamily.h"
+
+#define LOCTEXT_NAMESPACE "MetasoundStandardNodes_CatMixerNode"
+
+namespace Metasound
+{
+	namespace CatMixerPrivate
+	{
+		METASOUND_PARAM(InputGain,			"Gain", "Channel Agnostic Output");
+		METASOUND_PARAM(InputCat,			"Cat In", "Channel Agnostic Input");
+		METASOUND_PARAM(InputNumChannels,	"NumChannels", "Num Output Channels");
+		METASOUND_PARAM(OutputCat,			"Cat Out", "Channel Agnostic Output");
+		const FString DefaultCatFormat = TEXT("Mono_1_0");
+	}
+	
+	class FCatMixerOperator final : public TExecutableOperator<FCatMixerOperator>
+	{
+	public:
+		FCatMixerOperator(const FBuildOperatorParams& InParams, FChannelAgnosticTypeReadRef&& InInputCat, FFloatReadRef&& InGain, FInt32ReadRef&& InNumOutputChannels)
+			: Gain(MoveTemp(InGain))
+			, Inputs (MoveTemp(InInputCat))
+			, NumOutputChannels(MoveTemp(InNumOutputChannels))
+			, Outputs(FChannelAgnosticTypeWriteRef::CreateNew(InParams.OperatorSettings, CatMixerPrivate::DefaultCatFormat)) // initalize to a default format a head of binding. (maybe default init is better?)
+			, Settings(InParams.OperatorSettings)
+		{
+			Reset(InParams);
+		}
+		virtual ~FCatMixerOperator() override = default;
+
+		static const FVertexInterface& GetDefaultInterface()
+		{
+			using namespace CatMixerPrivate;
+			auto CreateDefaultInterface = []()-> FVertexInterface
+			{
+				// inputs
+				FInputVertexInterface InputInterface;
+				InputInterface.Add(TInputDataVertex<FChannelAgnosticType>(METASOUND_GET_PARAM_NAME_AND_METADATA(InputCat),CatMixerPrivate::DefaultCatFormat));
+				InputInterface.Add(TInputDataVertex<float>(METASOUND_GET_PARAM_NAME_AND_METADATA(InputGain)));
+				InputInterface.Add(TInputDataVertex<int32>(METASOUND_GET_PARAM_NAME_AND_METADATA(InputNumChannels)));
+				
+				// outputs
+				FOutputVertexInterface OutputInterface;
+				OutputInterface.Add(TOutputDataVertex<FChannelAgnosticType>(METASOUND_GET_PARAM_NAME_AND_METADATA(OutputCat)));
+				
+				return FVertexInterface(InputInterface, OutputInterface);
+			}; // end lambda: CreateDefaultInterface()
+
+			static const FVertexInterface DefaultInterface = CreateDefaultInterface();
+			return DefaultInterface;
+		}
+
+		static const FNodeClassMetadata& GetNodeInfo()
+		{
+			static const FNodeClassMetadata Metadata = CreateNodeClassMetadata
+			(
+				FName(TEXT("Cat Mixer Node")),
+				METASOUND_LOCTEXT("Metasound_CatMixerNodeDisplayName", "Cat Mixer Node"),
+				METASOUND_LOCTEXT("Metasound_CatMixerNodeDescription", "Cat Mixer Node"),
+				GetDefaultInterface()	
+			);
+			return Metadata;
+		}
+		
+		static TUniquePtr<IOperator> CreateOperator(const FBuildOperatorParams& InParams, FBuildResults& OutResults)
+		{
+			using namespace CatMixerPrivate;
+
+			const FInputVertexInterfaceData& InputData = InParams.InputData;
+			return MakeUnique<FCatMixerOperator>(
+				InParams,
+				InputData.GetOrCreateDefaultDataReadReference<FChannelAgnosticType>(METASOUND_GET_PARAM_NAME(InputCat), InParams.OperatorSettings),
+				InputData.GetOrCreateDefaultDataReadReference<float>(METASOUND_GET_PARAM_NAME(InputGain), InParams.OperatorSettings),
+				InputData.GetOrCreateDefaultDataReadReference<int32>(METASOUND_GET_PARAM_NAME(InputNumChannels), InParams.OperatorSettings)
+			);
+		}
+
+		virtual void BindInputs(FInputVertexInterfaceData& InOutVertexData) override
+		{
+			using namespace CatMixerPrivate;
+			InOutVertexData.BindReadVertex(METASOUND_GET_PARAM_NAME(InputGain) , Gain);
+			InOutVertexData.BindReadVertex(METASOUND_GET_PARAM_NAME(InputCat) , Inputs);
+			InOutVertexData.BindReadVertex(METASOUND_GET_PARAM_NAME(InputNumChannels) , NumOutputChannels);
+		}
+
+		virtual void BindOutputs(FOutputVertexInterfaceData& InOutVertexData) override
+		{
+			using namespace CatMixerPrivate;
+
+			if (const int32 NumChannels = *NumOutputChannels; NumChannels > 0) 
+			{
+				// Create cat type capable of holding N channels.
+				// FIXME: Use FName chooser.
+				if (const Audio::FChannelTypeFamily* Type = Audio::GetChannelRegistry().FindChannelType(NumChannels))
+				{
+					Outputs = FChannelAgnosticTypeWriteRef::CreateNew(Settings, Type->GetName().ToString());
+				}
+			
+				// Bind it
+				InOutVertexData.BindWriteVertex(METASOUND_GET_PARAM_NAME(OutputCat), Outputs);
+			}
+		}
+
+		void Reset(const IOperator::FResetParams& InParams)
+		{
+			PrevGains = *Gain;
+			Execute();
+		}
+
+		void Execute()
+		{
+		}
+		
+	private:
+
+		FFloatReadRef Gain;
+		FChannelAgnosticTypeReadRef Inputs;
+		FInt32ReadRef NumOutputChannels;
+		int32 NumFrames = 0;
+		int32 NumInputChannels = 0;
+		FChannelAgnosticTypeWriteRef Outputs;
+		FOperatorSettings Settings;
+		float PrevGains = 0.f;
+
+		static FNodeClassMetadata CreateNodeClassMetadata(const FName& InOperatorName, const FText& InDisplayName, const FText& InDescription, const FVertexInterface& InDefaultInterface)
+		{
+			FNodeClassMetadata Metadata
+			{
+				FNodeClassName { "CatAudioMixer", InOperatorName, FName() },
+				1, // Major Version
+				0, // Minor Version
+				InDisplayName,
+				InDescription,
+				PluginAuthor,
+				PluginNodeMissingPrompt,
+				InDefaultInterface,
+				{ NodeCategories::Mix },
+				{ METASOUND_LOCTEXT("Metasound_AudioMixerKeyword", "Mixer") },
+				FNodeDisplayStyle{}
+			};
+
+			return Metadata;
+		}
+	}; // class FCatMixerOperator
+
+	using FCatMixerNode = TNodeFacade<FCatMixerOperator>;
+
+} // namespace Metasound
+
+#undef LOCTEXT_NAMESPACE // "MetasoundStandardNodes_CatMixerNode"
diff --git a/Engine/Plugins/Experimental/MetasoundExperimental/Source/MetasoundExperimentalRuntime/Private/MetasoundCatMixerNode.h b/Engine/Plugins/Experimental/MetasoundExperimental/Source/MetasoundExperimentalRuntime/Private/MetasoundCatMixerNode.h
new file mode 100644
--- /dev/null
+++ b/Engine/Plugins/Experimental/MetasoundExperimental/Source/MetasoundExperimentalRuntime/Private/MetasoundCatMixerNode.h
@@ -0,0 +1,23 @@
+﻿// Copyright Epic Games, Inc. All Rights Reserved.
+
+#if 0
+
+#pragma once
+
+#include "MetasoundFacade.h"
+#include "MetasoundNodeInterface.h"
+#include "Templates/SharedPointer.h"
+
+namespace Metasound
+{
+	class FCatMixerNode final : public FNodeFacade
+	{
+	public:
+		explicit FCatMixerNode(FNodeData InNodeData, TSharedRef<const FNodeClassMetadata> InClassMetadata);
+		virtual ~FCatMixerNode() override = default;
+
+		static FNodeClassMetadata CreateNodeClassMetadata();
+	};
+}
+
+#endif //
\ No newline at end of file
diff --git a/Engine/Plugins/Experimental/MetasoundExperimental/Source/MetasoundExperimentalRuntime/Private/MetasoundCatTestingNode.cpp b/Engine/Plugins/Experimental/MetasoundExperimental/Source/MetasoundExperimentalRuntime/Private/MetasoundCatTestingNode.cpp
new file mode 100644
--- /dev/null
+++ b/Engine/Plugins/Experimental/MetasoundExperimental/Source/MetasoundExperimentalRuntime/Private/MetasoundCatTestingNode.cpp
@@ -0,0 +1,186 @@
+// Copyright Epic Games, Inc. All Rights Reserved.
+
+
+#include "MetasoundChannelAgnosticType.h"
+#include "MetasoundExecutableOperator.h"
+#include "MetasoundFacade.h"
+#include "MetasoundNodeRegistrationMacro.h"
+#include "MetasoundPrimitives.h"
+#include "MetasoundStandardNodesNames.h"
+#include "MetasoundVertex.h"
+#include "TypeFamily/ChannelTypeFamily.h"
+
+#define LOCTEXT_NAMESPACE "MetasoundStandardNodes_CatTestingNode"
+
+namespace Metasound
+{
+	namespace FCatTestingNodePrivate
+	{
+		METASOUND_PARAM(InputNumChannels, "NumChannels", "Number of cat channels to generate");
+		METASOUND_PARAM(OutputCat, "Output", "Channel Agnostic Output");
+	}
+
+	class FCatTestingOperator final : public TExecutableOperator<FCatTestingOperator>
+	{
+	public:
+		FCatTestingOperator(const FBuildOperatorParams& InParams, FInt32ReadRef&& InNumChannels)
+			: NumChannels(MoveTemp(InNumChannels))
+			, SampleRate(InParams.OperatorSettings.GetSampleRate())
+			, NumFramesPerBlock(InParams.OperatorSettings.GetNumFramesPerBlock())
+			, Outputs(FChannelAgnosticTypeWriteRef::CreateNew(InParams.OperatorSettings))	// Default create until we bind.
+			, Settings(InParams.OperatorSettings)
+		{}
+		virtual ~FCatTestingOperator() override = default;
+
+		static const FVertexInterface& GetDefaultInterface()
+		{
+			using namespace FCatTestingNodePrivate;
+			
+			auto CreateDefaultInterface = []()-> FVertexInterface
+			{
+				// inputs
+				FInputVertexInterface InputInterface;
+
+				InputInterface.Add(TInputConstructorVertex<int32>(METASOUND_GET_PARAM_NAME_AND_METADATA(InputNumChannels)));
+
+				// outputs
+				FOutputVertexInterface OutputInterface;
+				OutputInterface.Add(TOutputDataVertex<FChannelAgnosticType>(METASOUND_GET_PARAM_NAME_AND_METADATA(OutputCat)));
+
+				return FVertexInterface(InputInterface, OutputInterface);
+			}; // end lambda: CreateDefaultInterface()
+
+			static const FVertexInterface DefaultInterface = CreateDefaultInterface();
+			return DefaultInterface;
+		}
+
+		static const FVertexInterface& DeclareVertexInterface()
+		{
+			auto CreateDefaultInterface = []()-> FVertexInterface
+			{
+				using namespace FCatTestingNodePrivate;
+				static constexpr int32 DefaultNumChannels = 1;
+
+				const FInputVertexInterface InputInterface(
+					TInputDataVertex<int32>(METASOUND_GET_PARAM_NAME_AND_METADATA(InputNumChannels), DefaultNumChannels)
+				);
+				const FOutputVertexInterface OutputInterface;
+				return FVertexInterface(InputInterface, OutputInterface);
+			};
+
+			static const FVertexInterface DefaultInterface = CreateDefaultInterface();
+			return DefaultInterface;
+		}
+
+		static TUniquePtr<IOperator> CreateOperator(const FBuildOperatorParams& InParams, FBuildResults& OutResults)
+		{
+			using namespace FCatTestingNodePrivate;
+			const FInputVertexInterfaceData& InputData = InParams.InputData;
+			return MakeUnique<FCatTestingOperator>(
+				InParams,
+				InputData.GetOrCreateDefaultDataReadReference<int32>(FName(METASOUND_GET_PARAM_NAME(InputNumChannels)), InParams.OperatorSettings)
+			);
+		}
+
+		virtual void BindInputs(FInputVertexInterfaceData& InOutVertexData) override
+		{
+			using namespace FCatTestingNodePrivate;
+			InOutVertexData.BindReadVertex(METASOUND_GET_PARAM_NAME(InputNumChannels), NumChannels);
+		}
+
+		virtual void BindOutputs(FOutputVertexInterfaceData& InOutVertexData) override
+		{
+			using namespace FCatTestingNodePrivate;
+
+			// Find the matching format in registry. This should really use an FName chooser
+			// TODO: switch this over, in mean time, use channel count to find appropriate format
+			if (const Audio::FChannelTypeFamily* Type = Audio::GetChannelRegistry().FindChannelType(*NumChannels))
+			{
+				Outputs = FChannelAgnosticTypeWriteRef::CreateNew(Settings, Type->GetName().ToString());
+			}
+			
+			InOutVertexData.BindWriteVertex(METASOUND_GET_PARAM_NAME(OutputCat), Outputs);
+		}
+
+		void Reset(const FResetParams&)
+		{
+			Execute();
+		}
+
+		float MakeTone(TArrayView<float> InBuffer, const int32 InFreq) const
+		{
+			float* Dst = InBuffer.GetData();
+			const float PhaseStep = SampleRate / InFreq;;
+			float Phase = 0.f;
+
+			for (int32 Remaining = InBuffer.Num(); Remaining > 0; Remaining--)
+			{
+				*Dst++ = FMath::Sin(Phase * UE_TWO_PI);
+				Phase += PhaseStep;
+				Phase = Phase - static_cast<int32>(Phase);
+			}
+			return Phase;
+		}
+		
+		void Execute()
+		{
+			check(*NumChannels == Outputs->NumChannels());
+			check(NumFramesPerBlock == Outputs->NumFrames());
+
+			// Create N channels of increasing frequency tones.
+			for (int32 i = 0, Freq = 110; i < *NumChannels; ++i, Freq *=2)
+			{
+				MakeTone(Outputs->GetChannel(i),Freq);
+			}
+		}
+
+		static const FNodeClassMetadata& GetNodeInfo()
+		{
+			// used if NumChannels > 2
+			auto CreateNodeClassMetadataMultiChan = []() -> FNodeClassMetadata
+			{
+				const FName OperatorName = TEXT("Cat Testing Node");
+				const FText NodeDisplayName = METASOUND_LOCTEXT("CatTestingNodeDisplayName", "Cat Testing Node");
+				const FText NodeDescription = METASOUND_LOCTEXT("CatTestingNodeDescriptiong", "Cat Testing Node");
+				const FVertexInterface NodeInterface = GetDefaultInterface();
+				return  CreateNodeClassMetadata(OperatorName, NodeDisplayName, NodeDescription, NodeInterface);
+			};
+			static const FNodeClassMetadata Metadata = CreateNodeClassMetadataMultiChan();
+			return Metadata;
+		}
+
+
+	private:
+		FInt32ReadRef NumChannels;
+		float SampleRate = 0.f;
+		int32 NumFramesPerBlock = 0;
+		FChannelAgnosticTypeWriteRef Outputs;
+		FOperatorSettings Settings;
+
+		static FNodeClassMetadata CreateNodeClassMetadata(const FName& InOperatorName, const FText& InDisplayName, const FText& InDescription, const FVertexInterface& InDefaultInterface)
+		{
+			FNodeClassMetadata Metadata
+			{
+				FNodeClassName { StandardNodes::Namespace, InOperatorName, StandardNodes::AudioVariant },
+				1, // Major Version
+				0, // Minor Version
+				InDisplayName,
+				InDescription,
+				PluginAuthor,
+				PluginNodeMissingPrompt,
+				InDefaultInterface,
+				{},
+				{ METASOUND_LOCTEXT("Metasound_CatTestingNode", "Cat Testing") },
+				FNodeDisplayStyle{}
+			};
+
+			return Metadata;
+		}
+	}; // class FCatMixerOperator
+
+	using FCatTestingNode = TNodeFacade<FCatTestingOperator>;
+	METASOUND_REGISTER_NODE(FCatTestingNode);
+
+} // namespace Metasound
+
+#undef LOCTEXT_NAMESPACE // "MetasoundStandardNodes_CatMixerNode"
diff --git a/Engine/Plugins/Experimental/MetasoundExperimental/Source/MetasoundExperimentalRuntime/Private/MetasoundCatWaveWriterNode.cpp b/Engine/Plugins/Experimental/MetasoundExperimental/Source/MetasoundExperimentalRuntime/Private/MetasoundCatWaveWriterNode.cpp
new file mode 100644
--- /dev/null
+++ b/Engine/Plugins/Experimental/MetasoundExperimental/Source/MetasoundExperimentalRuntime/Private/MetasoundCatWaveWriterNode.cpp
@@ -0,0 +1,244 @@
+// Copyright Epic Games, Inc. All Rights Reserved.
+
+#include "Audio/SimpleWaveWriter.h"
+#include "ChannelAgnostic/ChannelAgnosticType.h"
+#include "ChannelAgnostic/ChannelAgnosticTypeUtils.h"
+#include "NumberedFileCache.h"
+#include "HAL/FileManager.h"
+#include "MetasoundBuildError.h"
+#include "MetasoundChannelAgnosticType.h"
+#include "MetasoundExecutableOperator.h"
+#include "MetasoundNodeRegistrationMacro.h"
+#include "MetasoundParamHelper.h"
+#include "MetasoundPrimitives.h"
+#include "MetasoundStandardNodesCategories.h"
+#include "MetasoundStandardNodesNames.h"
+#include "MetasoundVertex.h"
+#include "Misc/Paths.h"
+#include "Misc/ScopeLock.h"
+
+#define LOCTEXT_NAMESPACE "MetasoundStandardNodes_CatWaveWriterNode"
+
+namespace Metasound
+{
+	namespace Experimental
+	{
+		namespace WaveWriterVertexNames
+		{
+			METASOUND_PARAM(InEnabledPin, "Enabled", "If this wave writer is enabled or not.");
+			METASOUND_PARAM(InFilenamePrefixPin, "Filename Prefix", "Filename Prefix of file you are writing.");
+			METASOUND_PARAM(InCatPin, "InputCat", "Channel Agnostic Input");
+		}
+
+		class FFileWriteError final : public FBuildErrorBase
+		{
+		public:
+			static const FName ErrorType;
+
+			virtual ~FFileWriteError() override = default;
+
+			FFileWriteError(const FNode& InNode, const FString& InFilename)
+	#if WITH_EDITOR
+				: FBuildErrorBase(ErrorType, METASOUND_LOCTEXT_FORMAT("MetasoundFileWriterErrorDescription", "File Writer Error while trying to write '{0}'", FText::FromString(InFilename)))
+	#else 
+				: FBuildErrorBase(ErrorType, FText::GetEmpty())
+	#endif // WITH_EDITOR
+			{
+				AddNode(InNode);
+			}
+		};
+		const FName FFileWriteError::ErrorType = FName(TEXT("MetasoundFileWriterError"));
+
+		namespace WaveWriterOperatorPrivate
+		{
+			using namespace Audio;
+			
+			// Need to keep this outside the template so there's only 1
+			TSharedPtr<FNumberedFileCache> GetNameCache()
+			{
+				static const TCHAR* WaveExt = TEXT(".wav");
+
+				// Build cache of numbered files (do this once only).
+				static TSharedPtr<FNumberedFileCache> NumberedFileCacheSP = MakeShared<FNumberedFileCache>(*FPaths::AudioCaptureDir(), WaveExt, IFileManager::Get());
+				return NumberedFileCacheSP;
+			}
+
+			static const FString& GetDefaultFileName()
+			{
+				static const FString DefaultFileName = TEXT("Output");
+				return DefaultFileName;
+			}
+		}
+
+		class FCatWaveWriterOperator : public TExecutableOperator<FCatWaveWriterOperator>
+		{
+		public:
+			FCatWaveWriterOperator(const FOperatorSettings& InSettings, FChannelAgnosticTypeReadRef&& InAudioBuffers, FBoolReadRef&& InEnabled, const TSharedPtr<FNumberedFileCache, ESPMode::ThreadSafe>& InNumberedFileCache, FStringReadRef&& InFilenamePrefix)
+				: AudioInputs{MoveTemp(InAudioBuffers)}
+				, Enabled{MoveTemp(InEnabled)}
+				, NumberedFileCacheSP{InNumberedFileCache}
+				, FileNamePrefix{MoveTemp(InFilenamePrefix)}
+				, SampleRate{InSettings.GetSampleRate()}
+				, NumInputFrames(InSettings.GetNumFramesPerBlock())
+			{
+			}
+
+			virtual void BindInputs(FInputVertexInterfaceData& InOutVertexData) override
+			{
+				using namespace WaveWriterVertexNames;
+
+				InOutVertexData.BindReadVertex(METASOUND_GET_PARAM_NAME(InCatPin), AudioInputs);
+				InOutVertexData.BindReadVertex(METASOUND_GET_PARAM_NAME(InEnabledPin), Enabled);
+				InOutVertexData.BindReadVertex(METASOUND_GET_PARAM_NAME(InFilenamePrefixPin), FileNamePrefix);
+
+				// Now we have an input, we can ask how many channel it has.
+				NumInputChannels = AudioInputs->NumChannels();
+				InterleaveBuffer.SetNumZeroed(NumInputChannels * NumInputFrames);
+			}
+
+			virtual void BindOutputs(FOutputVertexInterfaceData&) override
+			{
+			}
+
+			static const FVertexInterface& DeclareVertexInterface()
+			{
+				auto CreateDefaultInterface = []()-> FVertexInterface
+				{
+					using namespace WaveWriterOperatorPrivate;
+					using namespace WaveWriterVertexNames;
+
+					// inputs
+					const FInputVertexInterface InputInterface(
+						TInputDataVertex<FString>(METASOUND_GET_PARAM_NAME_AND_METADATA(InFilenamePrefixPin), GetDefaultFileName()),
+						TInputDataVertex<bool>(METASOUND_GET_PARAM_NAME_AND_METADATA(InEnabledPin), true),
+						TInputDataVertex<FChannelAgnosticType>(METASOUND_GET_PARAM_NAME_AND_METADATA(InCatPin))
+					);
+					const FOutputVertexInterface OutputInterface;
+					return FVertexInterface(InputInterface, OutputInterface);
+				};
+
+				static const FVertexInterface DefaultInterface = CreateDefaultInterface();
+				return DefaultInterface;
+			}
+
+			static const FNodeClassMetadata& GetNodeInfo()
+			{
+				static const FNodeClassMetadata Metadata = CreateNodeClassMetadata
+				(
+					FName(TEXT("Cat Wave Writer")),
+					METASOUND_LOCTEXT("Metasound_CatWaveWriterNodeMultiChannelDisplayName", "Wave Writer Channel Agnostic"),
+					METASOUND_LOCTEXT("Metasound_CatWaveWriterNodeMultiDescription", "Write a CAT audio signal to disk"),
+					DeclareVertexInterface()
+				);
+				return Metadata;
+			}
+
+			static TUniquePtr<IOperator> CreateOperator(const FBuildOperatorParams& InParams, FBuildResults& OutResults);
+			
+			void Execute()
+			{
+				// Enabled and wasn't before? Enable.
+				if (!bIsEnabled && *Enabled)
+				{
+					Enable();
+				}
+				// Disabled but currently Enabled? Disable.
+				else if (bIsEnabled && !*Enabled)
+				{
+					Disable();
+				}
+
+				// If we have a valid writer and enabled.
+				if (Writer && *Enabled && NumInputChannels > 0)
+				{
+					Audio::FCatUtils::Interleave(*AudioInputs, InterleaveBuffer);
+					Writer->Write(InterleaveBuffer);
+				}
+			}
+
+			void Reset(const IOperator::FResetParams& InParams)
+			{
+				if (bIsEnabled)
+				{
+					Disable();
+				}
+			}
+
+		protected:
+			static FNodeClassMetadata CreateNodeClassMetadata(const FName& InOperatorName, const FText& InDisplayName, const FText& InDescription, const FVertexInterface& InDefaultInterface)
+			{
+				FNodeClassMetadata Metadata
+				{
+					FNodeClassName { StandardNodes::Namespace, InOperatorName, StandardNodes::AudioVariant },
+					1, // Major Version
+					1, // Minor Version
+					InDisplayName,
+					InDescription,
+					PluginAuthor,
+					PluginNodeMissingPrompt,
+					InDefaultInterface,
+					{ NodeCategories::Io },
+					{ METASOUND_LOCTEXT("Metasound_AudioMixerKeyword", "Writer") },
+					FNodeDisplayStyle{}
+				};
+				return Metadata;
+			}
+
+			void Enable()
+			{
+				if (ensure(!bIsEnabled) && NumInputChannels > 0)
+				{
+					bIsEnabled = true;
+					const FString Filename = NumberedFileCacheSP->GenerateNextNumberedFilename(*FileNamePrefix);
+					if (TUniquePtr<FArchive> Stream {IFileManager::Get().CreateFileWriter(*Filename, IO_WRITE) }; Stream.IsValid() )
+					{ 					
+						Writer = MakeUnique<Audio::FSimpleWaveWriter>(MoveTemp(Stream), SampleRate, NumInputChannels, true);
+					}
+				}
+			}
+			void Disable()
+			{
+				if (ensure(bIsEnabled))
+				{
+					bIsEnabled = false;
+					Writer.Reset();
+					NumInputChannels = 0;
+				}
+			}
+
+			FChannelAgnosticTypeReadRef AudioInputs;
+			TArray<float> InterleaveBuffer;
+			FBoolReadRef Enabled;
+			TUniquePtr<Audio::FSimpleWaveWriter> Writer;
+			TSharedPtr<FNumberedFileCache> NumberedFileCacheSP;
+			FStringReadRef FileNamePrefix;
+			float SampleRate = 0.f;
+			int32 NumInputChannels = 0;
+			int32 NumInputFrames = 0;
+			bool bIsEnabled = false;
+		};
+
+		TUniquePtr<IOperator> FCatWaveWriterOperator::CreateOperator(const FBuildOperatorParams& InParams, FBuildResults& OutResults)
+		{
+			using namespace WaveWriterOperatorPrivate;
+			using namespace WaveWriterVertexNames;
+
+			const FOperatorSettings& Settings = InParams.OperatorSettings;
+			const FInputVertexInterfaceData& InputData = InParams.InputData;
+
+			return MakeUnique<FCatWaveWriterOperator>(
+				Settings,
+				InputData.GetOrCreateDefaultDataReadReference<FChannelAgnosticType>(METASOUND_GET_PARAM_NAME(InCatPin),  Settings),
+				InputData.GetOrCreateDefaultDataReadReference<bool>(METASOUND_GET_PARAM_NAME(InEnabledPin), Settings),
+				GetNameCache(),
+				InputData.GetOrCreateDefaultDataReadReference<FString>(METASOUND_GET_PARAM_NAME(InFilenamePrefixPin), Settings)
+			);
+		}
+
+		using FCatWaveWriterNode = TNodeFacade<FCatWaveWriterOperator>;
+		METASOUND_REGISTER_NODE(FCatWaveWriterNode);
+	}// namespace Experimental
+
+} // namespace Metasounds
+
+#undef LOCTEXT_NAMESPACE
diff --git a/Engine/Plugins/Experimental/MetasoundExperimental/Source/MetasoundExperimentalRuntime/Private/MetasoundChannelAgnosticType.cpp b/Engine/Plugins/Experimental/MetasoundExperimental/Source/MetasoundExperimentalRuntime/Private/MetasoundChannelAgnosticType.cpp
new file mode 100644
--- /dev/null
+++ b/Engine/Plugins/Experimental/MetasoundExperimental/Source/MetasoundExperimentalRuntime/Private/MetasoundChannelAgnosticType.cpp
@@ -0,0 +1,37 @@
+﻿// Copyright Epic Games, Inc. All Rights Reserved.
+
+#include "MetasoundChannelAgnosticType.h"
+#include "TypeFamily/TypeFamily.h"
+#include "TypeFamily/ChannelTypeFamily.h"
+
+namespace Metasound
+{
+	// TODO: Pass this in via the operator settings some how. possibly on the environment.
+	//static Audio::TFamilyRegistry<Audio::FChannelTypeFamily> GetChannelTypeRegistry()
+	//{
+	//	using namespace Audio;
+	//	return TFamilyRegistry<FChannelTypeFamily>(GetChannelRegistry());
+	//}
+
+	/**
+	 * Returns the default (i.e. default constructed) format of a CAT. Static.
+	 * @return Name to Pass to Registry. 
+	 */
+	FName FChannelAgnosticType::GetDefaultCatFormat()
+	{
+		static const FName DefaultFormat(TEXT("Mono")); // 1 channel, Mono buffer.
+		check(Audio::GetChannelRegistry().FindConcreteChannel(DefaultFormat) != nullptr);
+		return DefaultFormat;
+	}
+
+	FChannelAgnosticType::FChannelAgnosticType(const FOperatorSettings& InSettings, const FName& InChannelTypeName)
+		: FChannelAgnosticType(InSettings.GetNumFramesPerBlock(), InChannelTypeName)
+	{}
+	
+	FChannelAgnosticType::FChannelAgnosticType(const int32 InNumFramesPerBlock, const FName& InChannelTypeName)
+		: Super(*Audio::GetChannelRegistry().FindConcreteChannel(InChannelTypeName), InNumFramesPerBlock)
+	{}
+}
+
+// Register.
+REGISTER_METASOUND_DATATYPE(Metasound::FChannelAgnosticType, "Cat");
\ No newline at end of file
diff --git a/Engine/Plugins/Experimental/MetasoundExperimental/Source/MetasoundExperimentalRuntime/Private/MetasoundExampleNodeConfiguration.cpp b/Engine/Plugins/Experimental/MetasoundExperimental/Source/MetasoundExperimentalRuntime/Private/MetasoundExampleNodeConfiguration.cpp
new file mode 100644
--- /dev/null
+++ b/Engine/Plugins/Experimental/MetasoundExperimental/Source/MetasoundExperimentalRuntime/Private/MetasoundExampleNodeConfiguration.cpp
@@ -0,0 +1,306 @@
+// Copyright Epic Games, Inc. All Rights Reserved.
+
+#include "MetasoundExampleNodeConfiguration.h" 
+
+#include "Internationalization/Text.h"
+#include "Math/UnrealMathUtility.h"
+
+#include "MetasoundDataFactory.h"
+#include "MetasoundFacade.h"
+#include "MetasoundOperatorData.h"
+#include "MetasoundNodeRegistrationMacro.h"
+#include "MetasoundExampleNodeConfiguration.h"
+#include "MetasoundExecutableOperator.h"
+#include "MetasoundParamHelper.h"
+#include "MetasoundTrigger.h"
+
+
+#define LOCTEXT_NAMESPACE "MetasoundExperimentalRuntime"
+
+namespace Metasound::Experimental
+{
+	namespace ExampleNodeConfigurationPrivate
+	{
+		const FLazyName InputBaseName{"In"};
+#if WITH_EDITOR
+		const FText InputTooltip = LOCTEXT("In_ToolTip", "A trigger");
+#else
+		const FText InputTooltip = FText::GetEmpty();
+#endif
+
+		FName MakeInputVertexName(int32 InIndex)
+		{
+			FName Name = InputBaseName;
+			Name.SetNumber(InIndex);
+			return Name;
+		}
+
+		FInputDataVertex MakeInputDataVertex(int32 InIndex)
+		{
+			FName InputName = MakeInputVertexName(InIndex);
+#if WITH_EDITOR
+			const FText InputDisplayName = FText::Format(LOCTEXT("In_DisplayName", "In {0}"), InIndex);
+#else
+			const FText InputDisplayName = FText::GetEmpty();
+#endif
+			return TInputDataVertex<FTrigger>{InputName, FDataVertexMetadata{InputTooltip, InputDisplayName}};
+		}
+
+		FLazyName OutputBaseName{"Out"};
+#if WITH_EDITOR
+		const FText OutputTooltip = LOCTEXT("Out_ToolTip", "A trigger");
+#else
+		const FText OutputTooltip = FText::GetEmpty();
+#endif
+
+		FName MakeOutputVertexName(int32 OutOutdex)
+		{
+			FName Name = OutputBaseName;
+			Name.SetNumber(OutOutdex);
+			return Name;
+		}
+
+		FOutputDataVertex MakeOutputDataVertex(int32 OutOutdex)
+		{
+			FName OutputName = MakeOutputVertexName(OutOutdex);
+#if WITH_EDITOR
+			const FText OutputDisplayName = FText::Format(LOCTEXT("Out_DisplayName", "Out {0}"), OutOutdex);
+#else
+			const FText OutputDisplayName = FText::GetEmpty();
+#endif
+			return TOutputDataVertex<FTrigger>{OutputName, FDataVertexMetadata{OutputTooltip, OutputDisplayName}};
+		}
+
+
+		// Create the node's vertex interface based upon the number of inputs and outputs
+		// desired. 
+		FVertexInterface GetVertexInterface(int32 InNumInputs, int32 InNumOutputs)
+		{
+			FInputVertexInterface InputInterface;
+			for (int32 i = 0; i < InNumInputs; i++)
+			{
+				InputInterface.Add(MakeInputDataVertex(i));
+			}
+
+			FOutputVertexInterface OutputInterface;
+			for (int32 i = 0; i < InNumOutputs; i++)
+			{
+				OutputInterface.Add(MakeOutputDataVertex(i));
+			}
+
+			return FVertexInterface
+				{
+					MoveTemp(InputInterface),
+					MoveTemp(OutputInterface)
+				};
+		}
+	}
+
+	/** To send data from a FMetaSoundFrontendNodeConfiguration to an IOperator, it should
+	 * be encapsulated in the form of a IOperatorData.  
+	 *
+	 * The use of the TOperatorData provides some saftey mechanisms for downcasting node configurations.
+	 */
+	class FExampleOperatorData : public TOperatorData<FExampleOperatorData>
+	{
+	public:
+		// The OperatorDataTypeName is used when downcasting an IOperatorData to ensure
+		// that the downcast is valid.
+		static const FLazyName OperatorDataTypeName;
+
+		FExampleOperatorData(const FString& InString)
+		: String(InString)
+		{
+		}
+
+		const FString& GetString() const
+		{
+			return String;
+		}
+
+	private:
+
+		FString String;
+	};
+
+	const FLazyName FExampleOperatorData::OperatorDataTypeName = "ExperimentalExampleOperatorData";
+
+	class FExampleConfigurableOperator : public TExecutableOperator<FExampleConfigurableOperator>
+	{
+	public:
+
+		FExampleConfigurableOperator(const FString& InString, TArray<TDataReadReference<FTrigger>> InInputTriggers, TArray<TDataWriteReference<FTrigger>> InOutputTriggers)
+		: InputTriggers(MoveTemp(InInputTriggers))
+		, OutputTriggers(MoveTemp(InOutputTriggers))
+		{
+			UE_LOG(LogMetaSound, Display, TEXT("Did the configurable string make it: %s"), *InString);
+		}
+
+		virtual void BindInputs(FInputVertexInterfaceData& InOutVertexData) override
+		{
+			using namespace ExampleNodeConfigurationPrivate;
+
+			for (int32 i = 0; i < InputTriggers.Num(); i++)
+			{
+				InOutVertexData.BindReadVertex(MakeInputVertexName(i), InputTriggers[i]);
+			}
+		}
+
+		virtual void BindOutputs(FOutputVertexInterfaceData& InOutVertexData) override
+		{
+			using namespace ExampleNodeConfigurationPrivate;
+
+			for (int32 i = 0; i < OutputTriggers.Num(); i++)
+			{
+				InOutVertexData.BindReadVertex(MakeOutputVertexName(i), OutputTriggers[i]);
+			}
+		}
+
+		void Execute()
+		{
+			// This node randomly rearranges all input triggers across all the 
+			// output triggers.
+			if (OutputTriggers.Num() > 0)
+			{
+				for (const TDataWriteReference<FTrigger>& OutputTrigger : OutputTriggers)
+				{
+					OutputTrigger->AdvanceBlock();
+				}
+				for (const TDataReadReference<FTrigger>& InputTrigger : InputTriggers)
+				{
+					InputTrigger->ExecuteBlock(
+							[](int32, int32) {},
+							[this](int32 InStartFrame, int32 InEndFrame)
+							{
+								int32 OutputIndex = FMath::RandRange(0, OutputTriggers.Num() - 1);
+								OutputTriggers[OutputIndex]->TriggerFrame(InStartFrame);
+							}
+					);
+				}
+			}
+		}
+
+		void Reset(const IOperator::FResetParams& InParams)
+		{
+			for (const TDataWriteReference<FTrigger>& OutputTrigger : OutputTriggers)
+			{
+				OutputTrigger->Reset();
+			}
+		}
+
+		static FNodeClassMetadata GetNodeInfo()
+		{
+			using namespace ExampleNodeConfigurationPrivate;
+			return FNodeClassMetadata
+			{
+				FNodeClassName{ "Experimental", "ConfigurableOperator", "" },
+				1, // Major version
+				0, // Minor version
+				LOCTEXT("ExampleConfigurableNodeName", "A Configurable Node"),	
+				LOCTEXT("ExampleConfigurableNodeDescription", "A Node which shows how to make a configurable node for yourself."),	
+				TEXT("UE"), // Author
+				LOCTEXT("ExampleConfigurablePromptIfMissing", "Enable the MetaSoundExperimental Plugin"), // Prompt if missing
+				ExampleNodeConfigurationPrivate::GetVertexInterface(1 /* default num inputs */, 1 /* default num outputs */),
+				{}
+			};
+		}
+
+		static TUniquePtr<IOperator> CreateOperator(const FBuildOperatorParams& InParams, FBuildResults& OutResults) 
+		{
+			using namespace ExampleNodeConfigurationPrivate;
+
+			// If your node configuration contains data that needs to be accessed by
+			// an operator, it can be accessed in the CreateOperator call. The 
+			// TSharedPtr<const IOperatorData> could even be stored on the IOperator. 
+			//
+			// In this example, the function `GetOperatorDataAs<>` safely downcasts
+			// the TSharedPtr<const IOperatorData> retrieved from the node. 
+			FString ConfiguredString = TEXT("Nope");
+			if (const FExampleOperatorData* ExampleConfig = CastOperatorData<const FExampleOperatorData>(InParams.Node.GetOperatorData().Get()))
+			{
+				ConfiguredString = ExampleConfig->GetString();
+			}
+
+			// If the node configuration overrides the default class interface, 
+			// the node's vertex interface will reflect the override. The vertex
+			// interface can be queried to see which inputs and outputs exist.
+			//
+			// For more complex scenarios, developers may want to pass other data
+			// through the IOperatorData. Alternatively, this node could have put
+			// the number of inputs and outputs into the IOperator data. 
+			const FVertexInterface& NodeInterface = InParams.Node.GetVertexInterface();
+
+			// Build the correct data references based upon what vertices exist
+			// on the NodeInterface
+			TArray<TDataReadReference<FTrigger>> InputTriggers;
+			int32 InputIndex = 0;
+			while (true)
+			{
+				FVertexName VertexName = MakeInputVertexName(InputIndex);
+				if (NodeInterface.ContainsInputVertex(VertexName))
+				{
+					InputTriggers.Add(InParams.InputData.GetOrCreateDefaultDataReadReference<FTrigger>(VertexName, InParams.OperatorSettings));
+				}
+				else
+				{
+					break;
+				}
+				InputIndex++;
+			}
+
+			TArray<TDataWriteReference<FTrigger>> OutputTriggers;
+			int32 OutputIndex = 0;
+			while (true)
+			{
+				FVertexName VertexName = MakeOutputVertexName(OutputIndex);
+				if (NodeInterface.ContainsOutputVertex(VertexName))
+				{
+					OutputTriggers.Add(TDataWriteReferenceFactory<FTrigger>::CreateExplicitArgs(InParams.OperatorSettings));
+				}
+				else
+				{
+					break;
+				}
+				OutputIndex++;
+			}
+
+
+			return MakeUnique<FExampleConfigurableOperator>(
+					ConfiguredString,
+					MoveTemp(InputTriggers),
+					MoveTemp(OutputTriggers)
+					);
+		}
+
+	private:
+		TArray<TDataReadReference<FTrigger>> InputTriggers;
+		TArray<TDataWriteReference<FTrigger>> OutputTriggers;
+	};
+
+	using FExampleConfigurableNode = TNodeFacade<FExampleConfigurableOperator>;
+
+	// The node extension must be registered along with the node. 
+	METASOUND_REGISTER_NODE_AND_CONFIGURATION(FExampleConfigurableNode, FMetaSoundExperimentalExampleNodeConfiguration);
+}
+
+FMetaSoundExperimentalExampleNodeConfiguration::FMetaSoundExperimentalExampleNodeConfiguration()
+: String("YES!")
+, NumInputs(1)
+, NumOutputs(1)
+{
+}
+
+TInstancedStruct<FMetasoundFrontendClassInterface> FMetaSoundExperimentalExampleNodeConfiguration::OverrideDefaultInterface(const FMetasoundFrontendClass& InNodeClass) const
+{
+	using namespace Metasound::Experimental::ExampleNodeConfigurationPrivate;
+	// Override the interface based upon the number of inputs and outputs desired. 
+	return TInstancedStruct<FMetasoundFrontendClassInterface>::Make(FMetasoundFrontendClassInterface::GenerateClassInterface(GetVertexInterface(NumInputs, NumOutputs)));
+}
+
+TSharedPtr<const Metasound::IOperatorData> FMetaSoundExperimentalExampleNodeConfiguration::GetOperatorData() const 
+{
+	// Any data the node configuration wishes to share with the operators can be produced here. 
+	return MakeShared<Metasound::Experimental::FExampleOperatorData>(String);
+}
+
+#undef LOCTEXT_NAMESPACE // MetasoundExperimentalRuntime
diff --git a/Engine/Plugins/Experimental/MetasoundExperimental/Source/MetasoundExperimentalRuntime/Private/MetasoundExperimentalRuntimeModule.cpp b/Engine/Plugins/Experimental/MetasoundExperimental/Source/MetasoundExperimentalRuntime/Private/MetasoundExperimentalRuntimeModule.cpp
new file mode 100644
--- /dev/null
+++ b/Engine/Plugins/Experimental/MetasoundExperimental/Source/MetasoundExperimentalRuntime/Private/MetasoundExperimentalRuntimeModule.cpp
@@ -0,0 +1,26 @@
+// Copyright Epic Games, Inc. All Rights Reserved.
+
+#include "MetasoundExperimentalRuntimeModule.h"
+
+#include "MetasoundCatMixerNode.h"
+#include "MetasoundChannelAgnosticType.h"
+#include "MetasoundDataTypeRegistrationMacro.h"
+#include "Misc/Paths.h"
+#include "Modules/ModuleManager.h"
+
+void FMetasoundExperimentalRuntimeModule::StartupModule()
+{
+	FMetasoundFrontendRegistryContainer::Get()->RegisterPendingNodes();
+
+	// Datatypes.
+	//Metasound::Frontend::RegisterDataType<Metasound::FChannelAgnosticType>();
+}
+
+void FMetasoundExperimentalRuntimeModule::ShutdownModule()
+{
+}
+IMPLEMENT_MODULE(FMetasoundExperimentalRuntimeModule, MetasoundExperimentalRuntime);
+
+//#if WITH_ENGINE
+//#error ("Being linked with Engine")
+//#endif //
diff --git a/Engine/Plugins/Experimental/MetasoundExperimental/Source/MetasoundExperimentalRuntime/Private/MetasoundExperimentalRuntimeModule.h b/Engine/Plugins/Experimental/MetasoundExperimental/Source/MetasoundExperimentalRuntime/Private/MetasoundExperimentalRuntimeModule.h
new file mode 100644
--- /dev/null
+++ b/Engine/Plugins/Experimental/MetasoundExperimental/Source/MetasoundExperimentalRuntime/Private/MetasoundExperimentalRuntimeModule.h
@@ -0,0 +1,17 @@
+// Copyright Epic Games, Inc. All Rights Reserved.
+
+#pragma once
+
+#include "Modules/ModuleInterface.h"
+
+
+class FMetasoundExperimentalRuntimeModule : public IModuleInterface
+{ 
+public:
+	//~ Begin IModuleInterface API
+	virtual void StartupModule() override;
+	virtual void ShutdownModule() override;
+	//~ End IModuleInterface API
+
+private:
+};
diff --git a/Engine/Plugins/Experimental/MetasoundExperimental/Source/MetasoundExperimentalRuntime/Public/MetasoundChannelAgnosticType.h b/Engine/Plugins/Experimental/MetasoundExperimental/Source/MetasoundExperimentalRuntime/Public/MetasoundChannelAgnosticType.h
new file mode 100644
--- /dev/null
+++ b/Engine/Plugins/Experimental/MetasoundExperimental/Source/MetasoundExperimentalRuntime/Public/MetasoundChannelAgnosticType.h
@@ -0,0 +1,41 @@
+// Copyright Epic Games, Inc. All Rights Reserved.
+
+#pragma once
+
+#include "ChannelAgnostic/ChannelAgnosticType.h"
+#include "MetasoundDataTypeRegistrationMacro.h"
+#include "DSP/MultiMono.h"
+
+namespace Metasound
+{
+	class FChannelAgnosticType : public Audio::FChannelAgnosticType
+	{
+	public:
+		using Super = Audio::FChannelAgnosticType;
+		using Super::Super;
+
+		/**
+		 * Returns the default (i.e. default constructed) format of a CAT.
+		 * @return Name to Pass to Registry. 
+		 */
+		static FName GetDefaultCatFormat();
+
+		explicit FChannelAgnosticType(const FOperatorSettings& InSettings, const FName& InChannelTypeName = GetDefaultCatFormat());
+		
+		explicit FChannelAgnosticType(const int32 InNumFramesPerBlock, const FName& InChannelTypeName = GetDefaultCatFormat());
+				
+		/* FLiteral compatible Constructor
+		 * TODO: Literals only support FString and not FNames, so there is a conversion here.
+		 */
+		explicit FChannelAgnosticType(const FOperatorSettings& InSettings, const FString& InChannelTypeName)
+			: FChannelAgnosticType(InSettings, FName(*InChannelTypeName))
+		{}
+	};
+	
+
+	
+}
+
+// Declare it.
+DECLARE_METASOUND_DATA_REFERENCE_TYPES(Metasound::FChannelAgnosticType, METASOUNDEXPERIMENTALRUNTIME_API,
+	FChannelAgnosticTypeTypeInfo, FChannelAgnosticTypeReadRef, FChannelAgnosticTypeWriteRef);
diff --git a/Engine/Plugins/Experimental/MetasoundExperimental/Source/MetasoundExperimentalRuntime/Public/MetasoundExampleNodeConfiguration.h b/Engine/Plugins/Experimental/MetasoundExperimental/Source/MetasoundExperimentalRuntime/Public/MetasoundExampleNodeConfiguration.h
new file mode 100644
--- /dev/null
+++ b/Engine/Plugins/Experimental/MetasoundExperimental/Source/MetasoundExperimentalRuntime/Public/MetasoundExampleNodeConfiguration.h
@@ -0,0 +1,31 @@
+// Copyright Epic Games, Inc. All Rights Reserved.
+
+#pragma once
+
+#include "MetasoundOperatorData.h"
+#include "MetasoundFrontendDocument.h"
+
+#include "MetasoundExampleNodeConfiguration.generated.h"
+
+USTRUCT()
+struct FMetaSoundExperimentalExampleNodeConfiguration : public FMetaSoundFrontendNodeConfiguration
+{
+	GENERATED_BODY()
+
+	FMetaSoundExperimentalExampleNodeConfiguration();
+
+	UPROPERTY(EditAnywhere, Category = General)
+	FString String;
+
+	UPROPERTY(EditAnywhere, Category = General, meta = (ClampMin = "1", ClampMax = "1000"))
+	uint32 NumInputs;
+
+	UPROPERTY(EditAnywhere, Category = General, meta = (ClampMin = "1", ClampMax = "1000"))
+	uint32 NumOutputs;
+
+	/* Get the current interface for the class based upon the node extension */
+	virtual TInstancedStruct<FMetasoundFrontendClassInterface> OverrideDefaultInterface(const FMetasoundFrontendClass& InNodeClass) const;
+
+	/* Pass data down to the operator. */
+	virtual TSharedPtr<const Metasound::IOperatorData> GetOperatorData() const override;
+};
diff --git a/Engine/Plugins/Experimental/MetasoundExperimental/Source/MetasoundExperimentalRuntime/Public/MetasoundFormatAgnosticType.h b/Engine/Plugins/Experimental/MetasoundExperimental/Source/MetasoundExperimentalRuntime/Public/MetasoundFormatAgnosticType.h
new file mode 100644
--- /dev/null
+++ b/Engine/Plugins/Experimental/MetasoundExperimental/Source/MetasoundExperimentalRuntime/Public/MetasoundFormatAgnosticType.h
@@ -0,0 +1,4 @@
+// Copyright Epic Games, Inc. All Rights Reserved.
+
+#pragma once
+
