[Experimental] CAT fix transcodes, fix wave writer. Add cheaper drop transcode options.

#tests LowLevelTests AudioUnitTests

diff --git a/Engine/Plugins/Experimental/MetasoundExperimental/Source/AudioExperimentalRuntime/Private/TypeFamily/ChannelTypeFamily.cpp b/Engine/Plugins/Experimental/MetasoundExperimental/Source/AudioExperimentalRuntime/Private/TypeFamily/ChannelTypeFamily.cpp
--- a/Engine/Plugins/Experimental/MetasoundExperimental/Source/AudioExperimentalRuntime/Private/TypeFamily/ChannelTypeFamily.cpp
+++ b/Engine/Plugins/Experimental/MetasoundExperimental/Source/AudioExperimentalRuntime/Private/TypeFamily/ChannelTypeFamily.cpp
@@ -145,6 +145,7 @@ namespace Audio
 	const FDiscreteChannelTypeFamily& InFromType, const FGetTranscoderParams& InParams) const
 	{
 		// Exact match? We can just memcpy each channel.
+		// TODO. in future these could be shared-ptrs from the main CAT memory block.
 		if (&InFromType == this)
 		{
 			return [NumChannels = InFromType.NumChannels()](TArrayView<const float*> Src, TArrayView<float*> Dst, const int32 NumFrames)
@@ -156,22 +157,47 @@ namespace Audio
 			};
 		}
 
-		// Make a mix matrix and call a mix/up down.
-		if (TArray<float> Gains; Create2DChannelMap(
-			{
-				.NumInputChannels = InFromType.NumChannels(),
-				.NumOutputChannels = NumChannels(),
-				.Order = EChannelMapOrder::OutputMajorOrder,
-				.MonoUpmixMethod = InParams.MixMethod, 
-				.bIsCenterChannelOnly = InFromType.NumChannels() == 1 && InFromType.HasSpeaker(ESpeakerShortNames::FC) 
-			}, Gains))
+		switch (InParams.TranscodeMethod)
 		{
-			return [MixGains = MoveTemp(Gains)](TArrayView<const float*> InSrc, TArrayView<float*> InDst, const int32 NumFrames)
+			case EChannelTranscodeMethod::ChannelDrop:
 			{
-				MultiMonoMixUpOrDown(InSrc,InDst,NumFrames,MixGains);
-			};
+				return [&](TArrayView<const float*> SrcChannels, TArrayView<float*> DstChannels, const int32 NumFrames)
+				{
+					// Copy everything destination wants, and nothing else.
+					ensure(Order.Num() == NumChannelsPrivate);
+					const TArray<ESpeakerShortNames> SrcOrder = InFromType.Order;
+					for (int32 i = 0; i < NumChannelsPrivate; ++i)
+					{
+						const ESpeakerShortNames DstSpeaker = Order[i];
+						if (const int32 SrcChannelIndex = SrcOrder.IndexOfByKey(DstSpeaker); SrcChannelIndex != INDEX_NONE)
+						{
+							FMemory::Memcpy(DstChannels[i], SrcChannels[SrcChannelIndex], NumFrames * sizeof(float));	
+						}
+					}
+				};
+			}
+			
+			case EChannelTranscodeMethod::MixUpOrDown:
+			{
+				// Make a mix matrix and call a mix/up down.
+				if (TArray<float> Gains; Create2DChannelMap(
+					{
+						.NumInputChannels = InFromType.NumChannels(),
+						.NumOutputChannels = NumChannels(),
+						.Order = EChannelMapOrder::OutputMajorOrder,
+						.MonoUpmixMethod = InParams.MixMethod,
+						.bIsCenterChannelOnly = InFromType.NumChannels() == 1 && InFromType.HasSpeaker(ESpeakerShortNames::FC)
+					}, Gains))
+				{
+					return [MixGains = MoveTemp(Gains)](TArrayView<const float*> InSrc, TArrayView<float*> InDst, const int32 NumFrames)
+					{
+						MultiMonoMixUpOrDown(InSrc, InDst, NumFrames, MixGains);
+					};
+				}
+			}
 		}
 
+		// fail.
 		return {};
 	}
 
@@ -249,6 +275,7 @@ Registry.RegisterType(TEXT(UE_STRINGIZE(NAME)),UE_JOIN(__Type,NAME))
 		REGISTER_CAT(Atmos7Dot1Dot2,	TEXT("Atmos"),			TEXT("Dolby Atmos (7.1.2)"),	(TArray{FL,FR,SL,SR,FC,BL,BR,TFL,TFR,LFE}),			false,		false);
 		REGISTER_CAT(Atmos7Dot1Dot4,	TEXT("Atmos"),			TEXT("Dolby Atmos (7.1.4)"),	(TArray{FL,FR,SL,SR,FC,BL,BR,TFL,TFR,TBL,TBR,LFE}),	true,		false);
 
+#if 0
 		// Channel packs.
 		REGISTER_CAT(ChannelPack,		TEXT("Cat"),				TEXT("Channel Pack"),			{},													false,		true);
 		
@@ -276,11 +303,13 @@ Registry.RegisterType(TEXT(UE_STRINGIZE(NAME)),UE_JOIN(__Type,NAME))
 			}
 			return Types;
 		}();
+#endif 
 					
 #define REGISTER_AMBISONICS(NAME, PARENT_NAME, FRIENDLY_NAME, ORDER, DEFAULT,ABSTRACT)\
 static FAmbisonicsChannelTypeFamily UE_JOIN(__Type,NAME)(UE_STRINGIZE(NAME), ORDER, Registry.FindChannel(PARENT_NAME), FRIENDLY_NAME, DEFAULT, ABSTRACT);\
 ensure(Registry.RegisterType(UE_STRINGIZE(NAME),UE_JOIN(__Type,NAME)))
 
+#if 0
 		// Ambisonics 
 		REGISTER_AMBISONICS(Ambisonics,				 TEXT("Cat"), 		TEXT("Ambisonics"), 							0, false, true);
 		REGISTER_AMBISONICS(AmbisonicsFirstOrder,	 TEXT("Ambisonics"), TEXT("First Order Ambisonics (4 channels)"), 	1, false, false); 
@@ -288,7 +317,7 @@ ensure(Registry.RegisterType(UE_STRINGIZE(NAME),UE_JOIN(__Type,NAME)))
 		REGISTER_AMBISONICS(AmbisonicsThirdOrder,	 TEXT("Ambisonics"), TEXT("Third Ambisonics (16 channels)"), 		3, false, false);
 		REGISTER_AMBISONICS(AmbisonicsFourthOrder,	 TEXT("Ambisonics"), TEXT("Fourth Ambisonics (25 chhanels)"), 		4, false, false); 
 		REGISTER_AMBISONICS(AmbisonicsFirthOrder,	 TEXT("Ambisonics"), TEXT("Firth Ambisonics (36 channels)"), 		5, false, false);  
-
+#endif //
 		// Channel packs.
 		// FPackedChannelTypeFamily MonoPack( TEXT("PackedChannel")
 		
diff --git a/Engine/Plugins/Experimental/MetasoundExperimental/Source/AudioExperimentalRuntime/Public/ScratchBuffer.h b/Engine/Plugins/Experimental/MetasoundExperimental/Source/AudioExperimentalRuntime/Public/ScratchBuffer.h
--- a/Engine/Plugins/Experimental/MetasoundExperimental/Source/AudioExperimentalRuntime/Public/ScratchBuffer.h
+++ b/Engine/Plugins/Experimental/MetasoundExperimental/Source/AudioExperimentalRuntime/Public/ScratchBuffer.h
@@ -101,14 +101,16 @@ namespace Audio
 
 		void DoAllocation(const int32 InNum)
 		{
-			check(InNum > 0);
-			check(Allocation.GetData() == nullptr);
-			Allocation = MakeArrayView<ValueType>(
-				static_cast<ValueType*>(Allocator->Malloc(InNum * sizeof(ValueType))),
-				InNum);
-			check(Allocation.GetData() != nullptr);
-			LifetimeToken = Allocator->GetCurrentLifetime();
-			FMemory::Memzero(Allocation.GetData(), Allocation.GetTypeSize() * Allocation.Num());
+			if (InNum > 0)
+			{
+				check(Allocation.GetData() == nullptr);
+				Allocation = MakeArrayView<ValueType>(
+					static_cast<ValueType*>(Allocator->Malloc(InNum * sizeof(ValueType))),
+					InNum);
+				check(Allocation.GetData() != nullptr);
+				LifetimeToken = Allocator->GetCurrentLifetime();
+				FMemory::Memzero(Allocation.GetData(), Allocation.GetTypeSize() * Allocation.Num());
+			}
 		}
 		void FreeAllocation()
 		{
diff --git a/Engine/Plugins/Experimental/MetasoundExperimental/Source/AudioExperimentalRuntime/Public/TypeFamily/ChannelTypeFamily.h b/Engine/Plugins/Experimental/MetasoundExperimental/Source/AudioExperimentalRuntime/Public/TypeFamily/ChannelTypeFamily.h
--- a/Engine/Plugins/Experimental/MetasoundExperimental/Source/AudioExperimentalRuntime/Public/TypeFamily/ChannelTypeFamily.h
+++ b/Engine/Plugins/Experimental/MetasoundExperimental/Source/AudioExperimentalRuntime/Public/TypeFamily/ChannelTypeFamily.h
@@ -7,8 +7,6 @@
 #include "Templates/Function.h"
 #include "Misc/Optional.h"
 
-//#include "ChannelTypeFamily.generated.h"
-
 /**
  * A Comprehensive Short form Speaker enumeration
  */
@@ -53,6 +51,12 @@ const TCHAR* LexToString(const ESpeakerShortNames InSpeaker);
 
 namespace Audio
 {
+	enum class EChannelTranscodeMethod : uint8
+	{
+		ChannelDrop,
+		MixUpOrDown,
+	};
+
 	// Forward declare these for visitor.
 	class FDiscreteChannelTypeFamily;
 	class FAmbisonicsChannelTypeFamily;
@@ -121,7 +125,6 @@ namespace Audio
 		 */
 		using FTranscoder = TFunction<void(TArrayView<const float*>, TArrayView<float*>, const int32 NumFrames)>;
 		
-		
 		/**
 		* Parameters to pass to GetTranscoder 
 		* To be expanded.
@@ -129,6 +132,7 @@ namespace Audio
 		struct FGetTranscoderParams
 		{
 			const FChannelTypeFamily& ToType;
+			EChannelTranscodeMethod TranscodeMethod = EChannelTranscodeMethod::ChannelDrop;
 			EChannelMapMonoUpmixMethod MixMethod = EChannelMapMonoUpmixMethod::EqualPower;
 		};
 		/**
diff --git a/Engine/Plugins/Experimental/MetasoundExperimental/Source/MetasoundExperimentalRuntime/Private/MetasoundCatBreakNode.cpp b/Engine/Plugins/Experimental/MetasoundExperimental/Source/MetasoundExperimentalRuntime/Private/MetasoundCatBreakNode.cpp
--- a/Engine/Plugins/Experimental/MetasoundExperimental/Source/MetasoundExperimentalRuntime/Private/MetasoundCatBreakNode.cpp
+++ b/Engine/Plugins/Experimental/MetasoundExperimental/Source/MetasoundExperimentalRuntime/Private/MetasoundCatBreakNode.cpp
@@ -80,15 +80,26 @@ namespace Metasound
 			// that the downcast is valid.
 			static const FLazyName OperatorDataTypeName;
 
-			explicit FCatBreakOperatorData(const FName& InToTypeName)
+			explicit FCatBreakOperatorData(const FName& InToTypeName, const Audio::EChannelTranscodeMethod InMethod, const Audio::EChannelMapMonoUpmixMethod InMixMethod)
 				: ToTypeName(InToTypeName)
+				, TranscodeMethod(InMethod)
 			{}
 			const FName& GetToType() const
 			{
 				return ToTypeName;
 			}
+			Audio::EChannelTranscodeMethod GetTranscodeMethod() const
+			{
+				return TranscodeMethod;
+			}
+			Audio::EChannelMapMonoUpmixMethod GetMixMethod() const
+			{
+				return MixMethod;
+			}
 		private:
 			FName ToTypeName;
+			Audio::EChannelTranscodeMethod TranscodeMethod = Audio::EChannelTranscodeMethod::ChannelDrop;
+			Audio::EChannelMapMonoUpmixMethod MixMethod = Audio::EChannelMapMonoUpmixMethod::EqualPower; 
 		};
 
 		// Linkage.
@@ -97,7 +108,8 @@ namespace Metasound
 		// Helper to create array of multi-mono channel pointers from a CAT.
 		Audio::TStackArrayOfPointers<float> MakeMultiMonoPointersFromBufferArray(const TArray<TDataWriteReference<FAudioBuffer>>& InArrayOfCat)
 		{
-			Audio::TStackArrayOfPointers<float> Result {};
+			Audio::TStackArrayOfPointers<float> Result;
+			Result.SetNum(InArrayOfCat.Num());
 			for (int32 i = 0; i < InArrayOfCat.Num(); ++i)
 			{
 				Result[i] = InArrayOfCat[i]->GetData();
@@ -110,7 +122,8 @@ namespace Metasound
 	{
 	public:
 		using FTranscoder = Audio::FChannelTypeFamily::FTranscoder;
-		FCatBreakOperator(const FBuildOperatorParams& InParams, FChannelAgnosticTypeReadRef&& InInputCat, TArray<TDataWriteReference<FAudioBuffer>>&& InOutputVerticies, FTranscoder&& InTranscoder, const FName InFormat)
+		FCatBreakOperator(const FBuildOperatorParams& InParams, FChannelAgnosticTypeReadRef&& InInputCat, TArray<TDataWriteReference<FAudioBuffer>>&& InOutputVerticies,
+			FTranscoder&& InTranscoder, const FName InFormat)
 			: InputCat(MoveTemp(InInputCat))
 			, OutputAudioVertices(MoveTemp(InOutputVerticies))
 			, Settings(InParams.OperatorSettings)
@@ -145,8 +158,7 @@ namespace Metasound
 			const FCatBreakOperatorData* CatBreakData = CastOperatorData<const FCatBreakOperatorData>(InParams.Node.GetOperatorData().Get());
 			if (!CatBreakData)
 			{
-				// Phil, this fires in the node test.
-				return {};
+				return MakeUnique<FNoOpOperator>();
 			}
 			
 			const FName OutputFormat = CatBreakData->GetToType();
@@ -155,7 +167,7 @@ namespace Metasound
 			if (!ConcreteToType)
 			{
 				// Throw editor error, we don't know the format.
-				return {};
+				return MakeUnique<FNoOpOperator>();
 			}
 
 			// Create the input pin.
@@ -172,16 +184,18 @@ namespace Metasound
 			// Always ask for the transcoder 
 			// In the trivial case where we are the same format, this will be just a memcpy.
 			FTranscoder Transcoder = InputPin->GetType().GetTranscoder(
-				{	
-					.ToType = *ConcreteToType
-				});
+				{
+					.ToType = *ConcreteToType,
+					.TranscodeMethod = CatBreakData->GetTranscodeMethod(),
+					.MixMethod = CatBreakData->GetMixMethod(),
+			});
 			
 			return MakeUnique<FCatBreakOperator>(
 				InParams,
 				MoveTemp(InputPin),
 				MoveTemp(OutputAudioVertices),
 				MoveTemp(Transcoder),
-				ConcreteToType->GetName()	
+				ConcreteToType->GetName()
 			);
 		}
 
@@ -212,11 +226,7 @@ namespace Metasound
 			{
 				TStackArrayOfPointers<const float> Src = MakeMultiMonoPointersFromView(InputCat->GetRawMultiMono(), Settings.GetNumFramesPerBlock(), InputCat->NumChannels());
 				TStackArrayOfPointers<float> Dst = MakeMultiMonoPointersFromBufferArray(OutputAudioVertices);
-				Transcoder(
-					MakeArrayView(Src.GetData(),InputCat->NumChannels()),
-					MakeArrayView(Dst.GetData(), OutputAudioVertices.Num()),
-					Settings.GetNumFramesPerBlock()
-					);
+				Transcoder(Src, Dst, Settings.GetNumFramesPerBlock() );
 			}
 		}
 		
@@ -260,7 +270,11 @@ TInstancedStruct<FMetasoundFrontendClassInterface> FMetaSoundCatBreakNodeConfigu
 
 TSharedPtr<const Metasound::IOperatorData> FMetaSoundCatBreakNodeConfiguration::GetOperatorData() const
 {
-	return MakeShared<Metasound::CatBreakNodePrivate::FCatBreakOperatorData>(Format);
+	return MakeShared<Metasound::CatBreakNodePrivate::FCatBreakOperatorData>(
+		Format,
+		static_cast<Audio::EChannelTranscodeMethod>(TranscodeMethod),
+		static_cast<Audio::EChannelMapMonoUpmixMethod>(MixMethod)
+	);
 }
 
 #undef LOCTEXT_NAMESPACE // 
diff --git a/Engine/Plugins/Experimental/MetasoundExperimental/Source/MetasoundExperimentalRuntime/Private/MetasoundCatBreakNode.h b/Engine/Plugins/Experimental/MetasoundExperimental/Source/MetasoundExperimentalRuntime/Private/MetasoundCatBreakNode.h
--- a/Engine/Plugins/Experimental/MetasoundExperimental/Source/MetasoundExperimentalRuntime/Private/MetasoundCatBreakNode.h
+++ b/Engine/Plugins/Experimental/MetasoundExperimental/Source/MetasoundExperimentalRuntime/Private/MetasoundCatBreakNode.h
@@ -4,6 +4,7 @@
 
 #include "MetasoundFrontendDocument.h"
 #include "Templates/SharedPointer.h"
+#include "MetasoundCatCastingNode.h"
 
 #include "MetasoundCatBreakNode.generated.h"
 
@@ -17,6 +18,12 @@ struct FMetaSoundCatBreakNodeConfiguration : public FMetaSoundFrontendNodeConfig
 	UPROPERTY(EditAnywhere, Category = General, meta = (GetOptions="MetasoundCatCastingOptionsHelper.GetCastingOptions"))
 	FName Format = TEXT("Stereo");
 
+	UPROPERTY(EditAnywhere, Category = General)
+	EMetasoundCatCastingMethod TranscodeMethod = EMetasoundCatCastingMethod::ChannelDrop;
+
+	UPROPERTY(EditAnywhere, Category = General, meta = (EditCondition = "TranscodeMethod == EMetasoundCatCastingMethod::MixUpOrDown", EditConditionHides))
+	EMetasoundChannelMapMonoUpmixMethod MixMethod = EMetasoundChannelMapMonoUpmixMethod::EqualPower; 
+
 	/* Get the current interface for the class based upon the node extension */
 	virtual TInstancedStruct<FMetasoundFrontendClassInterface> OverrideDefaultInterface(const FMetasoundFrontendClass& InNodeClass) const override;
 
diff --git a/Engine/Plugins/Experimental/MetasoundExperimental/Source/MetasoundExperimentalRuntime/Private/MetasoundCatCastingNode.cpp b/Engine/Plugins/Experimental/MetasoundExperimental/Source/MetasoundExperimentalRuntime/Private/MetasoundCatCastingNode.cpp
--- a/Engine/Plugins/Experimental/MetasoundExperimental/Source/MetasoundExperimentalRuntime/Private/MetasoundCatCastingNode.cpp
+++ b/Engine/Plugins/Experimental/MetasoundExperimental/Source/MetasoundExperimentalRuntime/Private/MetasoundCatCastingNode.cpp
@@ -29,15 +29,21 @@ namespace Metasound
 			// that the downcast is valid.
 			static const FLazyName OperatorDataTypeName;
 
-			explicit FCatCastingOperatorData(const FName& InToTypeName)
+			explicit FCatCastingOperatorData(const FName& InToTypeName, Audio::EChannelTranscodeMethod InTranscodeMethod, Audio::EChannelMapMonoUpmixMethod InMixMethod)
 				: ToTypeName(InToTypeName)
+				, TranscodeMethod(InTranscodeMethod)
+				, MixMethod(InMixMethod)
 			{}
 			const FName& GetToType() const
 			{
 				return ToTypeName;
 			}
+			Audio::EChannelMapMonoUpmixMethod GetMixMethod() const { return MixMethod; }
+			Audio::EChannelTranscodeMethod GetTranscodeMethod() const { return TranscodeMethod; }
 		private:
 			FName ToTypeName;
+			Audio::EChannelTranscodeMethod TranscodeMethod;
+			Audio::EChannelMapMonoUpmixMethod MixMethod;
 		};
 
 		// Linkage.
@@ -47,13 +53,14 @@ namespace Metasound
 	class FCatCastingOperator final : public TExecutableOperator<FCatCastingOperator>
 	{
 	public:
-		FCatCastingOperator(const FBuildOperatorParams& InParams, FChannelAgnosticTypeReadRef&& InInputCat,const FName InToName)
+		FCatCastingOperator(const FBuildOperatorParams& InParams, FChannelAgnosticTypeReadRef&& InInputCat, const CatCastingPrivate::FCatCastingOperatorData& InData, const FName InConcreteName)
 			: InputFrom(MoveTemp(InInputCat))
-			, ToFormatName(InToName)
-			, OutputCastResult(FChannelAgnosticTypeWriteRef::CreateNew(InParams.OperatorSettings, InToName))
+			, ToFormatName(InConcreteName)
+			, OutputCastResult(FChannelAgnosticTypeWriteRef::CreateNew(InParams.OperatorSettings, InConcreteName))
 			, Settings(InParams.OperatorSettings)
-		{
-		}
+			, TranscodeMethod(InData.GetTranscodeMethod())
+			, MixMethod(InData.GetMixMethod())
+		{}
 		virtual ~FCatCastingOperator() override = default;
 
 		static const FVertexInterface& GetDefaultInterface()
@@ -81,7 +88,10 @@ namespace Metasound
 			using namespace CatCastingPrivate;
 
 			const FCatCastingOperatorData* CatTestingConfigData = CastOperatorData<const FCatCastingOperatorData>(InParams.Node.GetOperatorData().Get());
-			check(CatTestingConfigData);
+			if (!CatTestingConfigData)
+			{
+				return MakeUnique<FNoOpOperator>();
+			}
 			const FName CastToName = CatTestingConfigData->GetToType();
 			
 			const FInputVertexInterfaceData& InputData = InParams.InputData;
@@ -96,6 +106,7 @@ namespace Metasound
 			return MakeUnique<FCatCastingOperator>(
 				InParams,
 				MoveTemp(InputCat),
+				*CatTestingConfigData,
 				CastToNameSane
 			);
 		}
@@ -113,7 +124,11 @@ namespace Metasound
 
 			// Create transcoder 
 			Transcoder = InputFrom->GetType().GetTranscoder(
-				{ .ToType = OutputCastResult->GetType() }
+				{
+					.ToType = OutputCastResult->GetType(),
+					.TranscodeMethod = TranscodeMethod,
+					.MixMethod = MixMethod,
+				}
 			);
 		}
 
@@ -126,13 +141,11 @@ namespace Metasound
 		{
 			if (Transcoder)
 			{
-				auto Src = Audio::MakeMultiMonoPointersFromView(InputFrom->GetRawMultiMono(), Settings.GetNumFramesPerBlock(), InputFrom->NumChannels());
-				auto Dst = Audio::MakeMultiMonoPointersFromView(OutputCastResult->GetRawMultiMono(), Settings.GetNumFramesPerBlock(), OutputCastResult->NumChannels());				
-				Transcoder(
-				 	MakeArrayView(Src.GetData(),InputFrom->NumChannels()),
-				 	MakeArrayView(Dst.GetData(),OutputCastResult->NumChannels()),
-				 	InputFrom->NumFrames()
-				 	);
+				using namespace CatCastingPrivate;
+				using namespace Audio;
+				TStackArrayOfPointers<const float> Src = MakeMultiMonoPointersFromView(InputFrom->GetRawMultiMono(), Settings.GetNumFramesPerBlock(), InputFrom->NumChannels());
+				TStackArrayOfPointers<float> Dst = MakeMultiMonoPointersFromView(OutputCastResult->GetRawMultiMono(), Settings.GetNumFramesPerBlock(), OutputCastResult->NumChannels());
+				Transcoder(Src, Dst,  Settings.GetNumFramesPerBlock());
 			}
 		}
 		static FNodeClassMetadata GetNodeInfo()
@@ -158,6 +171,8 @@ namespace Metasound
 		FChannelAgnosticTypeWriteRef OutputCastResult;
 		FOperatorSettings Settings;
 		FTranscoder Transcoder;
+		Audio::EChannelTranscodeMethod TranscodeMethod;
+		Audio::EChannelMapMonoUpmixMethod MixMethod;
 	}; // class FCatCastingOperator
 		
 	using FCatCastingNode = TNodeFacade<FCatCastingOperator>;
@@ -184,7 +199,11 @@ TInstancedStruct<FMetasoundFrontendClassInterface> FMetaSoundCatCastingNodeConfi
 
 TSharedPtr<const Metasound::IOperatorData> FMetaSoundCatCastingNodeConfiguration::GetOperatorData() const
 {
-	return MakeShared<Metasound::CatCastingPrivate::FCatCastingOperatorData>(ToType);
+	return MakeShared<Metasound::CatCastingPrivate::FCatCastingOperatorData>(
+		ToType,
+		static_cast<Audio::EChannelTranscodeMethod>(TranscodeMethod),
+		static_cast<Audio::EChannelMapMonoUpmixMethod>(MixMethod)
+		);
 }
 
 #undef LOCTEXT_NAMESPACE // 
diff --git a/Engine/Plugins/Experimental/MetasoundExperimental/Source/MetasoundExperimentalRuntime/Private/MetasoundCatCastingNode.h b/Engine/Plugins/Experimental/MetasoundExperimental/Source/MetasoundExperimentalRuntime/Private/MetasoundCatCastingNode.h
--- a/Engine/Plugins/Experimental/MetasoundExperimental/Source/MetasoundExperimentalRuntime/Private/MetasoundCatCastingNode.h
+++ b/Engine/Plugins/Experimental/MetasoundExperimental/Source/MetasoundExperimentalRuntime/Private/MetasoundCatCastingNode.h
@@ -4,6 +4,7 @@
 
 #include "MetasoundFrontendDocument.h"
 #include "Templates/SharedPointer.h"
+#include "TypeFamily/ChannelTypeFamily.h"
 #include "UObject/PropertyText.h"
 
 #include "MetasoundCatCastingNode.generated.h"
@@ -17,6 +18,21 @@ class UMetasoundCatCastingOptionsHelper : public UObject
 	static TArray<FPropertyTextFName> GetCastingOptions();
 };
 
+UENUM()
+enum class EMetasoundCatCastingMethod : uint8
+{
+	ChannelDrop = static_cast<uint8>(Audio::EChannelTranscodeMethod::ChannelDrop),
+	MixUpOrDown = static_cast<uint8>(Audio::EChannelTranscodeMethod::MixUpOrDown)
+};
+
+UENUM()
+enum class EMetasoundChannelMapMonoUpmixMethod : uint8
+{
+	Linear = static_cast<uint8>(Audio::EChannelMapMonoUpmixMethod::Linear),
+	EqualPower = static_cast<uint8>(Audio::EChannelMapMonoUpmixMethod::EqualPower),
+	FullVolume = static_cast<uint8>(Audio::EChannelMapMonoUpmixMethod::FullVolume)
+};
+
 USTRUCT()
 struct FMetaSoundCatCastingNodeConfiguration : public FMetaSoundFrontendNodeConfiguration
 {
@@ -27,8 +43,11 @@ struct FMetaSoundCatCastingNodeConfiguration : public FMetaSoundFrontendNodeConf
 	UPROPERTY(EditAnywhere, Category = General, meta = (GetOptions="MetasoundCatCastingOptionsHelper.GetCastingOptions"))
 	FName ToType = TEXT("Mono");
 	
-	UPROPERTY(EditAnywhere, Category = General, meta = (ClampMin = "1", ClampMax = "1000"))
-	uint32 NumOutputs = 1;
+	UPROPERTY(EditAnywhere, Category = General)
+	EMetasoundCatCastingMethod TranscodeMethod = EMetasoundCatCastingMethod::ChannelDrop;
+
+	UPROPERTY(EditAnywhere, Category = General, meta = (EditCondition = "TranscodeMethod == EMetasoundCatCastingMethod::MixUpOrDown", EditCOonditionHides))
+	EMetasoundChannelMapMonoUpmixMethod MixMethod = EMetasoundChannelMapMonoUpmixMethod::EqualPower; 
 
 	/* Get the current interface for the class based upon the node extension */
 	virtual TInstancedStruct<FMetasoundFrontendClassInterface> OverrideDefaultInterface(const FMetasoundFrontendClass& InNodeClass) const override;
diff --git a/Engine/Plugins/Experimental/MetasoundExperimental/Source/MetasoundExperimentalRuntime/Private/MetasoundCatMakeNode.cpp b/Engine/Plugins/Experimental/MetasoundExperimental/Source/MetasoundExperimentalRuntime/Private/MetasoundCatMakeNode.cpp
--- a/Engine/Plugins/Experimental/MetasoundExperimental/Source/MetasoundExperimentalRuntime/Private/MetasoundCatMakeNode.cpp
+++ b/Engine/Plugins/Experimental/MetasoundExperimental/Source/MetasoundExperimentalRuntime/Private/MetasoundCatMakeNode.cpp
@@ -126,7 +126,10 @@ namespace Metasound
 			using namespace CatMakeNodePrivate;
 
 			const FCatMakeOperatorData* CatMakeData = CastOperatorData<const FCatMakeOperatorData>(InParams.Node.GetOperatorData().Get());
-			check(CatMakeData);
+			if (!CatMakeData)
+			{
+				return TUniquePtr<FNoOpOperator>();
+			}
 			const FName OutputFormat = CatMakeData->GetToType();
 			const FInputVertexInterfaceData& InputData = InParams.InputData;
 			const Audio::FChannelTypeFamily* ConcreteToType = Audio::GetChannelRegistry().FindConcreteChannel(OutputFormat);
diff --git a/Engine/Plugins/Experimental/MetasoundExperimental/Source/MetasoundExperimentalRuntime/Private/MetasoundCatTestingNode.cpp b/Engine/Plugins/Experimental/MetasoundExperimental/Source/MetasoundExperimentalRuntime/Private/MetasoundCatTestingNode.cpp
--- a/Engine/Plugins/Experimental/MetasoundExperimental/Source/MetasoundExperimentalRuntime/Private/MetasoundCatTestingNode.cpp
+++ b/Engine/Plugins/Experimental/MetasoundExperimental/Source/MetasoundExperimentalRuntime/Private/MetasoundCatTestingNode.cpp
@@ -1,6 +1,7 @@
 // Copyright Epic Games, Inc. All Rights Reserved.
 
 
+#include "MetasoundCatTestingNode.h"
 #include "MetasoundChannelAgnosticType.h"
 #include "MetasoundExecutableOperator.h"
 #include "MetasoundFacade.h"
@@ -18,53 +19,66 @@ namespace Metasound
 	{
 		METASOUND_PARAM(InputNumChannels, "NumChannels", "Number of cat channels to generate");
 		METASOUND_PARAM(OutputCat, "Output", "Channel Agnostic Output");
+
+		FVertexInterface MakeClassInterface(const FName Format)
+		{
+			const Audio::FChannelTypeFamily* Found = Audio::GetChannelRegistry().FindConcreteChannel(Format);
+			if (!Found)
+			{
+				return {};
+			}
+			FInputVertexInterface InputInterface;
+			FOutputVertexInterface OutputInterface;
+			OutputInterface.Add(TOutputDataVertex<FChannelAgnosticType>(METASOUND_GET_PARAM_NAME_AND_METADATA(OutputCat)));
+			return{
+				MoveTemp(InputInterface),
+				MoveTemp(OutputInterface)
+			};
+		}
+		
+		class FCatTestingOperatorData final : public TOperatorData<FCatTestingOperatorData>
+		{
+		public:
+			// The OperatorDataTypeName is used when downcasting an IOperatorData to ensure
+			// that the downcast is valid.
+			static const FLazyName OperatorDataTypeName;
+
+			explicit FCatTestingOperatorData(const FName& InToTypeName)
+				: ToTypeName(InToTypeName)
+			{}
+			const FName& GetToType() const
+			{
+				return ToTypeName;
+			}
+		private:
+			FName ToTypeName;
+		};
+
+		// Linkage.
+		const FLazyName FCatTestingOperatorData::OperatorDataTypeName = TEXT("FCatTestingOperatorData");
 	}
 
 	class FCatTestingOperator final : public TExecutableOperator<FCatTestingOperator>
 	{
 	public:
-		FCatTestingOperator(const FBuildOperatorParams& InParams, FInt32ReadRef&& InNumChannels)
-			: NumChannels(MoveTemp(InNumChannels))
-			, SampleRate(InParams.OperatorSettings.GetSampleRate())
+		FCatTestingOperator(const FBuildOperatorParams& InParams, const FName InConcreteFormat)
+			: SampleRate(InParams.OperatorSettings.GetSampleRate())
 			, NumFramesPerBlock(InParams.OperatorSettings.GetNumFramesPerBlock())
-			, Outputs(FChannelAgnosticTypeWriteRef::CreateNew(InParams.OperatorSettings))	// Default create until we bind.
+			, Outputs(FChannelAgnosticTypeWriteRef::CreateNew(InParams.OperatorSettings, InConcreteFormat))	
 			, Settings(InParams.OperatorSettings)
 		{}
 		virtual ~FCatTestingOperator() override = default;
 
-		static const FVertexInterface& GetDefaultInterface()
-		{
-			using namespace FCatTestingNodePrivate;
-			
-			auto CreateDefaultInterface = []()-> FVertexInterface
-			{
-				// inputs
-				FInputVertexInterface InputInterface;
-
-				InputInterface.Add(TInputConstructorVertex<int32>(METASOUND_GET_PARAM_NAME_AND_METADATA(InputNumChannels)));
-
-				// outputs
-				FOutputVertexInterface OutputInterface;
-				OutputInterface.Add(TOutputDataVertex<FChannelAgnosticType>(METASOUND_GET_PARAM_NAME_AND_METADATA(OutputCat)));
-
-				return FVertexInterface(InputInterface, OutputInterface);
-			}; // end lambda: CreateDefaultInterface()
-
-			static const FVertexInterface DefaultInterface = CreateDefaultInterface();
-			return DefaultInterface;
-		}
-
 		static const FVertexInterface& DeclareVertexInterface()
 		{
 			auto CreateDefaultInterface = []()-> FVertexInterface
 			{
 				using namespace FCatTestingNodePrivate;
-				static constexpr int32 DefaultNumChannels = 1;
-
-				const FInputVertexInterface InputInterface(
-					TInputDataVertex<int32>(METASOUND_GET_PARAM_NAME_AND_METADATA(InputNumChannels), DefaultNumChannels)
-				);
-				const FOutputVertexInterface OutputInterface;
+				const FInputVertexInterface InputInterface;
+				const FOutputVertexInterface OutputInterface
+				{
+					TOutputDataVertex<FChannelAgnosticType>(METASOUND_GET_PARAM_NAME_AND_METADATA(OutputCat))
+				};
 				return FVertexInterface(InputInterface, OutputInterface);
 			};
 
@@ -75,31 +89,30 @@ namespace Metasound
 		static TUniquePtr<IOperator> CreateOperator(const FBuildOperatorParams& InParams, FBuildResults& OutResults)
 		{
 			using namespace FCatTestingNodePrivate;
-			const FInputVertexInterfaceData& InputData = InParams.InputData;
+
+			const FCatTestingOperatorData* TestingOperatorData = CastOperatorData<const FCatTestingOperatorData>(InParams.Node.GetOperatorData().Get());
+			if (!TestingOperatorData)
+			{
+				return TUniquePtr<FNoOpOperator>();
+			}
+			const Audio::FChannelTypeFamily* Concrete = Audio::GetChannelRegistry().FindConcreteChannel(TestingOperatorData->GetToType());
+			if (!Concrete)
+			{
+				return TUniquePtr<FNoOpOperator>();
+			}
+			
 			return MakeUnique<FCatTestingOperator>(
 				InParams,
-				InputData.GetOrCreateDefaultDataReadReference<int32>(FName(METASOUND_GET_PARAM_NAME(InputNumChannels)), InParams.OperatorSettings)
+				Concrete->GetName()
 			);
 		}
 
 		virtual void BindInputs(FInputVertexInterfaceData& InOutVertexData) override
-		{
-			using namespace FCatTestingNodePrivate;
-			InOutVertexData.BindReadVertex(METASOUND_GET_PARAM_NAME(InputNumChannels), NumChannels);
-		}
+		{}
 
 		virtual void BindOutputs(FOutputVertexInterfaceData& InOutVertexData) override
 		{
-			using namespace FCatTestingNodePrivate;
-
-			// Find the matching format in registry. This should really use an FName chooser
-			// TODO: switch this over, in mean time, use channel count to find appropriate format
-			if (const Audio::FChannelTypeFamily* Type = Audio::GetChannelRegistry().FindChannelType(*NumChannels))
-			{
-				Outputs = FChannelAgnosticTypeWriteRef::CreateNew(Settings, Type->GetName().ToString());
-			}
-			
-			InOutVertexData.BindWriteVertex(METASOUND_GET_PARAM_NAME(OutputCat), Outputs);
+			InOutVertexData.BindWriteVertex(FCatTestingNodePrivate::OutputCatName, Outputs);
 		}
 
 		void Reset(const FResetParams&)
@@ -124,11 +137,11 @@ namespace Metasound
 		
 		void Execute()
 		{
-			check(*NumChannels == Outputs->NumChannels());
 			check(NumFramesPerBlock == Outputs->NumFrames());
 
 			// Create N channels of increasing frequency tones.
-			for (int32 i = 0, Freq = 110; i < *NumChannels; ++i, Freq *=2)
+			const int32 NumChannels = Outputs->NumChannels();;
+			for (int32 i = 0, Freq = 110; i < NumChannels; ++i, Freq *=2)
 			{
 				MakeTone(Outputs->GetChannel(i),Freq);
 			}
@@ -139,10 +152,10 @@ namespace Metasound
 			// used if NumChannels > 2
 			auto CreateNodeClassMetadataMultiChan = []() -> FNodeClassMetadata
 			{
-				const FName OperatorName = TEXT("Cat Testing Node");
-				const FText NodeDisplayName = METASOUND_LOCTEXT("CatTestingNodeDisplayName", "Cat Testing Node");
-				const FText NodeDescription = METASOUND_LOCTEXT("CatTestingNodeDescriptiong", "Cat Testing Node");
-				const FVertexInterface NodeInterface = GetDefaultInterface();
+				const FName OperatorName = TEXT("CAT Testing Node");
+				const FText NodeDisplayName = METASOUND_LOCTEXT("CatTestingNodeDisplayName", "CAT Testing Node");
+				const FText NodeDescription = METASOUND_LOCTEXT("CatTestingNodeDescription", "CAT Testing Node");
+				const FVertexInterface NodeInterface =  DeclareVertexInterface();
 				return  CreateNodeClassMetadata(OperatorName, NodeDisplayName, NodeDescription, NodeInterface);
 			};
 			static const FNodeClassMetadata Metadata = CreateNodeClassMetadataMultiChan();
@@ -151,7 +164,6 @@ namespace Metasound
 
 
 	private:
-		FInt32ReadRef NumChannels;
 		float SampleRate = 0.f;
 		int32 NumFramesPerBlock = 0;
 		FChannelAgnosticTypeWriteRef Outputs;
@@ -179,8 +191,21 @@ namespace Metasound
 	}; // class FCatMixerOperator
 
 	using FCatTestingNode = TNodeFacade<FCatTestingOperator>;
-	METASOUND_REGISTER_NODE(FCatTestingNode);
+	// Disabled for now.
+	//METASOUND_REGISTER_NODE_AND_CONFIGURATION(FCatTestingNode, FMetaSoundCatTestingNodeConfiguration);
+
 
 } // namespace Metasound
 
+TInstancedStruct<FMetasoundFrontendClassInterface> FMetaSoundCatTestingNodeConfiguration::OverrideDefaultInterface(
+	const FMetasoundFrontendClass& InNodeClass) const
+{
+	return  TInstancedStruct<FMetasoundFrontendClassInterface>::Make(FMetasoundFrontendClassInterface::GenerateClassInterface(Metasound::FCatTestingNodePrivate::MakeClassInterface(ToType)));
+}
+
+TSharedPtr<const Metasound::IOperatorData> FMetaSoundCatTestingNodeConfiguration::GetOperatorData() const
+{
+	return MakeShared<Metasound::FCatTestingNodePrivate::FCatTestingOperatorData>(ToType);
+}
+
 #undef LOCTEXT_NAMESPACE // "MetasoundStandardNodes_CatMixerNode"
diff --git a/Engine/Plugins/Experimental/MetasoundExperimental/Source/MetasoundExperimentalRuntime/Private/MetasoundCatTestingNode.h b/Engine/Plugins/Experimental/MetasoundExperimental/Source/MetasoundExperimentalRuntime/Private/MetasoundCatTestingNode.h
new file mode 100644
--- /dev/null
+++ b/Engine/Plugins/Experimental/MetasoundExperimental/Source/MetasoundExperimentalRuntime/Private/MetasoundCatTestingNode.h
@@ -0,0 +1,25 @@
+// Copyright Epic Games, Inc. All Rights Reserved.
+
+#pragma once
+
+#include "MetasoundFrontendDocument.h"
+#include "Templates/SharedPointer.h"
+
+#include "MetasoundCatTestingNode.generated.h"
+
+USTRUCT()
+struct FMetaSoundCatTestingNodeConfiguration : public FMetaSoundFrontendNodeConfiguration
+{
+	GENERATED_BODY()
+
+	FMetaSoundCatTestingNodeConfiguration() =  default;
+
+	UPROPERTY(EditAnywhere, Category = General, meta = (GetOptions="MetasoundCatCastingOptionsHelper.GetCastingOptions"))
+	FName ToType = TEXT("Mono");
+
+	/* Get the current interface for the class based upon the node extension */
+	virtual TInstancedStruct<FMetasoundFrontendClassInterface> OverrideDefaultInterface(const FMetasoundFrontendClass& InNodeClass) const override;
+
+	/* Pass data down to the operator. */
+	virtual TSharedPtr<const Metasound::IOperatorData> GetOperatorData() const override;
+};
diff --git a/Engine/Plugins/Experimental/MetasoundExperimental/Source/MetasoundExperimentalRuntime/Private/MetasoundCatWaveWriterNode.cpp b/Engine/Plugins/Experimental/MetasoundExperimental/Source/MetasoundExperimentalRuntime/Private/MetasoundCatWaveWriterNode.cpp
--- a/Engine/Plugins/Experimental/MetasoundExperimental/Source/MetasoundExperimentalRuntime/Private/MetasoundCatWaveWriterNode.cpp
+++ b/Engine/Plugins/Experimental/MetasoundExperimental/Source/MetasoundExperimentalRuntime/Private/MetasoundCatWaveWriterNode.cpp
@@ -70,35 +70,24 @@ namespace Metasound
 			}
 		}
 
-		class FCatWaveWriterOperator : public TExecutableOperator<FCatWaveWriterOperator>
+		class FCatWaveWriterOperator final : public TExecutableOperator<FCatWaveWriterOperator>
 		{
 		public:
-			FCatWaveWriterOperator(const FOperatorSettings& InSettings, FChannelAgnosticTypeReadRef&& InAudioBuffers, FBoolReadRef&& InEnabled, const TSharedPtr<FNumberedFileCache, ESPMode::ThreadSafe>& InNumberedFileCache, FStringReadRef&& InFilenamePrefix)
+			FCatWaveWriterOperator(const FBuildOperatorParams& InParams, FChannelAgnosticTypeReadRef&& InAudioBuffers, FBoolReadRef&& InEnabled, const TSharedPtr<FNumberedFileCache, ESPMode::ThreadSafe>& InNumberedFileCache, FStringReadRef&& InFilenamePrefix)
 				: AudioInputs{MoveTemp(InAudioBuffers)}
 				, Enabled{MoveTemp(InEnabled)}
 				, NumberedFileCacheSP{InNumberedFileCache}
 				, FileNamePrefix{MoveTemp(InFilenamePrefix)}
-				, SampleRate{InSettings.GetSampleRate()}
-				, NumInputFrames(InSettings.GetNumFramesPerBlock())
+				, SampleRate{InParams.OperatorSettings.GetSampleRate()}
+				, NumInputFrames{InParams.OperatorSettings.GetNumFramesPerBlock()}
 			{
-			}
-
-			virtual void BindInputs(FInputVertexInterfaceData& InOutVertexData) override
-			{
-				using namespace WaveWriterVertexNames;
-
-				InOutVertexData.BindReadVertex(METASOUND_GET_PARAM_NAME(InCatPin), AudioInputs);
-				InOutVertexData.BindReadVertex(METASOUND_GET_PARAM_NAME(InEnabledPin), Enabled);
-				InOutVertexData.BindReadVertex(METASOUND_GET_PARAM_NAME(InFilenamePrefixPin), FileNamePrefix);
-
 				// Now we have an input, we can ask how many channel it has.
 				NumInputChannels = AudioInputs->NumChannels();
 				InterleaveBuffer.SetNumZeroed(NumInputChannels * NumInputFrames);
 			}
 
-			virtual void BindOutputs(FOutputVertexInterfaceData&) override
-			{
-			}
+			virtual void BindInputs(FInputVertexInterfaceData& InOutVertexData) override {}
+			virtual void BindOutputs(FOutputVertexInterfaceData&) override {}
 
 			static const FVertexInterface& DeclareVertexInterface()
 			{
@@ -227,7 +216,7 @@ namespace Metasound
 			const FInputVertexInterfaceData& InputData = InParams.InputData;
 
 			return MakeUnique<FCatWaveWriterOperator>(
-				Settings,
+				InParams,
 				InputData.GetOrCreateDefaultDataReadReference<FChannelAgnosticType>(METASOUND_GET_PARAM_NAME(InCatPin),  Settings),
 				InputData.GetOrCreateDefaultDataReadReference<bool>(METASOUND_GET_PARAM_NAME(InEnabledPin), Settings),
 				GetNameCache(),
diff --git a/Engine/Plugins/Runtime/Metasound/Source/MetasoundEngineTest/Private/EngineTestMetaSoundAutomatedNodeTest_Bind.cpp b/Engine/Plugins/Runtime/Metasound/Source/MetasoundEngineTest/Private/EngineTestMetaSoundAutomatedNodeTest_Bind.cpp
--- a/Engine/Plugins/Runtime/Metasound/Source/MetasoundEngineTest/Private/EngineTestMetaSoundAutomatedNodeTest_Bind.cpp
+++ b/Engine/Plugins/Runtime/Metasound/Source/MetasoundEngineTest/Private/EngineTestMetaSoundAutomatedNodeTest_Bind.cpp
@@ -75,6 +75,7 @@ bool FMetasoundAutomatedNodeTest_Bind::RunTest(const FString& InRegistryKeyStrin
 	if (!Operator.IsValid())
 	{
 		AddError(FString::Printf(TEXT("Failed to create operator from node %s - %s."), *InRegistryKeyString, *GetPrettyName(RegistryKey)));
+		return false;
 	}
 
 	// Query the data on the operator
diff --git a/Engine/Plugins/Runtime/Metasound/Source/MetasoundEngineTest/Private/EngineTestMetaSoundAutomatedNodeTest_Reset.cpp b/Engine/Plugins/Runtime/Metasound/Source/MetasoundEngineTest/Private/EngineTestMetaSoundAutomatedNodeTest_Reset.cpp
--- a/Engine/Plugins/Runtime/Metasound/Source/MetasoundEngineTest/Private/EngineTestMetaSoundAutomatedNodeTest_Reset.cpp
+++ b/Engine/Plugins/Runtime/Metasound/Source/MetasoundEngineTest/Private/EngineTestMetaSoundAutomatedNodeTest_Reset.cpp
@@ -84,6 +84,7 @@ bool FMetasoundAutomatedNodeTest_Reset::RunTest(const FString& InRegistryKeyStri
 		if (!Operator.IsValid())
 		{
 			AddError(FString::Printf(TEXT("Failed to create operator from node %s - %s."), *InRegistryKeyString, *GetPrettyName(RegistryKey)));
+			return;
 		}
 
 		// Store a copy of the input values data values so the inputs
diff --git a/Engine/Source/Runtime/SignalProcessing/Private/MultiMono.cpp b/Engine/Source/Runtime/SignalProcessing/Private/MultiMono.cpp
--- a/Engine/Source/Runtime/SignalProcessing/Private/MultiMono.cpp
+++ b/Engine/Source/Runtime/SignalProcessing/Private/MultiMono.cpp
@@ -16,17 +16,6 @@ namespace Audio
 			MixGains);
 	}
 
-	static void MixMono(TArrayView<const float> InSrc, TArrayView<float> InDst, const float InGain)
-	{
-		const float* Src = InSrc.GetData();
-		float* Dst = InDst.GetData();
-		const int32 Num = InSrc.Num();
-		for (int32 i = 0; i < Num; i++,Src++,Dst++)
-		{
-			*Dst += *Src * InGain;
-		}
-	}
-
 	void MultiMonoMixUpOrDown(TArrayView<const float*> InSrc, TArrayView<float*> InDst, const int32 NumFrames,TArrayView<const float> MixGains)
 	{
 		const int32 NumSrcChannels = InSrc.Num();
@@ -36,21 +25,14 @@ namespace Audio
 		
 		for (int32 DstCh = 0; DstCh < NumDstChannels; ++DstCh)
 		{
+			const TArrayView<float> Dst = MakeArrayView(InDst[DstCh], NumFrames);
+			FMemory::Memzero(Dst.GetData(), NumFrames * sizeof(float));
 			for (int32 SrcCh = 0; SrcCh < NumSrcChannels; ++SrcCh)
 			{
 				if (const float ChannelGain = MixGains[DstCh + (SrcCh * NumDstChannels)]; !FMath::IsNearlyZero(ChannelGain))
 				{
-					MixMono(
-					MakeArrayView(InSrc[SrcCh], NumFrames),
-					MakeArrayView(InDst[DstCh], NumFrames),
-						ChannelGain);
-					
-					/*
-					ArrayMixIn(
-						MakeArrayView(InSrc[SrcCh], NumFrames),
-						MakeArrayView(InDst[DstCh], NumFrames),
-						ChannelGain);
-						*/
+					const TArrayView<const float> Src = MakeArrayView(InSrc[SrcCh], NumFrames);
+					ArrayMixIn(Src, Dst, ChannelGain);
 				}
 			}
 		}	
diff --git a/Engine/Source/Runtime/SignalProcessing/Public/DSP/MultiMono.h b/Engine/Source/Runtime/SignalProcessing/Public/DSP/MultiMono.h
--- a/Engine/Source/Runtime/SignalProcessing/Public/DSP/MultiMono.h
+++ b/Engine/Source/Runtime/SignalProcessing/Public/DSP/MultiMono.h
@@ -14,20 +14,21 @@
 
 namespace Audio
 {
-	constexpr int32 MaxNumTranscodeChannels = 128; // Some relatively high number. These are just pointers on stack.
+	constexpr int32 MaxStackMultiMonoChannels = 128; // Some relatively high number. These are just pointers on stack.
 
 	// Stack array of pointers alias.
-	template<typename T> using TStackArrayOfPointers = TStaticArray<T*,MaxNumTranscodeChannels>;
+	template<typename T> using TStackArrayOfPointers = TArray<T*,TInlineAllocator<MaxStackMultiMonoChannels>>;
 
 	// Helper to create array of pointers from multi-mono.
 	// Warning Array of pointers will be larger than you need, but filled to InNumChannels
 	template<typename T> TStackArrayOfPointers<T> MakeMultiMonoPointersFromView(TArrayView<T> InMultiMono, const int32 InNumFrames, const int32 InNumChannels)
 	{
-		checkSlow(InNumChannels <= MaxNumTranscodeChannels);
+		checkSlow(InNumChannels <= MaxStackMultiMonoChannels);
 		checkSlow(InNumChannels > 0);
 		checkSlow(InMultiMono.GetData());
 		checkSlow(InMultiMono.Num() % InNumFrames == 0 && InMultiMono.Num() >= InNumFrames);
-		TStackArrayOfPointers<T> ArrayOfPointers = {};
+		TStackArrayOfPointers<T> ArrayOfPointers;
+		ArrayOfPointers.SetNum(InNumChannels);
 		for (int32 i = 0; i < InNumChannels; ++i)
 		{
 			ArrayOfPointers[i] = InMultiMono.GetData() + (InNumFrames * i);
