MetaSound node configuration editor details customization improvements
- Users can write their own details customizations to register with a specific subclass of FMetaSoundNodeConfiguration
- Experimental example node and customization for a custom details panel widget updating the node's operator data live while previewing
- Reduce unnecessary details refresh and interface updates

diff --git a/Engine/Plugins/Experimental/MetasoundExperimental/Source/MetasoundExperimentalEditor/MetasoundExperimentalEditor.build.cs b/Engine/Plugins/Experimental/MetasoundExperimental/Source/MetasoundExperimentalEditor/MetasoundExperimentalEditor.build.cs
--- a/Engine/Plugins/Experimental/MetasoundExperimental/Source/MetasoundExperimentalEditor/MetasoundExperimentalEditor.build.cs
+++ b/Engine/Plugins/Experimental/MetasoundExperimental/Source/MetasoundExperimentalEditor/MetasoundExperimentalEditor.build.cs
@@ -37,7 +37,16 @@ public class MetasoundExperimentalEditor : ModuleRules
 			{
 				"CoreUObject",
 				"Engine",
-
+				"MetasoundExperimentalRuntime",
+				"PropertyEditor",
+				"UnrealEd",
+				"Slate",
+				"SlateCore",
+				"AudioWidgets",
+				"MetasoundEditor",
+				"StructUtilsEditor",
+				"MetasoundEngine",
+				"MetasoundFrontend"
 				// ... add private dependencies that you statically link with here ...	
 			}
 		);
diff --git a/Engine/Plugins/Experimental/MetasoundExperimental/Source/MetasoundExperimentalEditor/Private/MetasoundExampleNodeDetailsCustomization.cpp b/Engine/Plugins/Experimental/MetasoundExperimental/Source/MetasoundExperimentalEditor/Private/MetasoundExampleNodeDetailsCustomization.cpp
new file mode 100644
--- /dev/null
+++ b/Engine/Plugins/Experimental/MetasoundExperimental/Source/MetasoundExperimentalEditor/Private/MetasoundExampleNodeDetailsCustomization.cpp
@@ -0,0 +1,103 @@
+// Copyright Epic Games, Inc. All Rights Reserved.
+
+#define LOCTEXT_NAMESPACE "MetasoundExperimentalEditor"
+#include "MetasoundExampleNodeDetailsCustomization.h" 
+
+#include "DetailWidgetRow.h"
+#include "IDetailChildrenBuilder.h"
+#include "InstancedStructDetails.h"
+#include "MetasoundEditorGraphNode.h"
+#include "MetasoundBuilderBase.h"
+#include "MetasoundExampleNodeConfiguration.h"
+#include "MetasoundNodeConfigurationCustomization.h"
+#include "PropertyEditorModule.h"
+#include "Widgets/Input/SSlider.h"
+
+const FString StructIdentifierWithDelimiters =  TEXT(".Struct.");
+
+FExampleWidgetNodeConfigurationCustomization::FExampleWidgetNodeConfigurationCustomization(TSharedPtr<IPropertyHandle> InStructProperty, TWeakObjectPtr<UMetasoundEditorGraphNode> InNode)
+	: Metasound::Editor::FMetaSoundNodeConfigurationDataDetails(InStructProperty, InNode)
+{
+	if (InStructProperty && InStructProperty->IsValidHandle())
+	{
+		StructPropertyPath = InStructProperty->GeneratePathToProperty();
+	}
+}
+
+void FExampleWidgetNodeConfigurationCustomization::OnChildRowAdded(IDetailPropertyRow& ChildRow)
+{
+	TSharedPtr<IPropertyHandle> ChildHandle = ChildRow.GetPropertyHandle();
+	if (!ChildHandle || !ChildHandle->IsValidHandle())
+	{
+		return;
+	}
+
+	const FString PropertyPath = ChildHandle->GeneratePathToProperty();
+	// Customize a specific member 
+	if (PropertyPath == StructPropertyPath + StructIdentifierWithDelimiters + GET_MEMBER_NAME_CHECKED(FMetaSoundWidgetExampleNodeConfiguration, MyFloat).ToString())
+	{
+		MyFloatPropertyHandle = ChildHandle;
+
+		TSharedPtr<SWidget> NameWidget;
+		TSharedPtr<SWidget> ValueWidget;
+		FDetailWidgetRow Row;
+
+		ChildRow.GetDefaultWidgets(NameWidget, ValueWidget, Row);
+
+		ChildRow.CustomWidget(true)
+			.NameContent()
+			[
+				NameWidget.ToSharedRef()
+			]
+			.ValueContent()
+			[
+				SNew(SSlider)
+				.MinValue(0.0f)
+				.MaxValue(1.0f)
+				.Value_Lambda([Handle=ChildRow.GetPropertyHandle()]{
+					float Value = 0.0f;
+					if (Handle.IsValid() && Handle->IsValidHandle())
+					{
+						Handle->GetValue(Value);
+					}
+					return Value;
+				})
+				.OnValueChanged_Lambda([Handle = ChildRow.GetPropertyHandle()](float Value)
+				{
+					if (Handle.IsValid() && Handle->IsValidHandle())
+					{
+						Handle->SetValue(Value);
+					}
+				})
+			];
+	}
+	// Add custom onvalue changed
+	TDelegate<void(const FPropertyChangedEvent&)> OnValueChangedDelegate = TDelegate<void(const FPropertyChangedEvent&)>::CreateSP(this, &FExampleWidgetNodeConfigurationCustomization::OnChildPropertyChanged);
+	ChildHandle->SetOnPropertyValueChangedWithData(OnValueChangedDelegate);
+
+	// Add base class on value changed
+	Metasound::Editor::FMetaSoundNodeConfigurationDataDetails::OnChildRowAdded(ChildRow);
+}
+
+void FExampleWidgetNodeConfigurationCustomization::OnChildPropertyChanged(const FPropertyChangedEvent& InPropertyChangedEvent)
+{
+	if (InPropertyChangedEvent.GetPropertyName() == GET_MEMBER_NAME_CHECKED(FMetaSoundWidgetExampleNodeConfiguration, MyFloat))
+	{
+		if (GraphNode.IsValid() && MyFloatPropertyHandle.IsValid() && MyFloatPropertyHandle->IsValidHandle())
+		{
+			FMetaSoundFrontendDocumentBuilder& DocBuilder = GraphNode->GetBuilderChecked().GetBuilder();
+			const FGuid& NodeID = GraphNode->GetNodeID();
+
+			// Update the operator data value from the configuration property handle value
+			// Node configuration operator data API is experimental
+			// so this code will be made cleaner in the future 
+			TInstancedStruct<FMetaSoundFrontendNodeConfiguration> Config = DocBuilder.FindNodeConfiguration(NodeID);
+			TSharedPtr<const Metasound::IOperatorData> OperatorData = Config.Get().GetOperatorData();
+			const TSharedPtr<const Metasound::Experimental::FWidgetExampleOperatorData> WidgetOperatorData = StaticCastSharedPtr<const Metasound::Experimental::FWidgetExampleOperatorData>(OperatorData);
+			TSharedPtr<Metasound::Experimental::FWidgetExampleOperatorData> MutableWidgetOperatorData = ConstCastSharedPtr<Metasound::Experimental::FWidgetExampleOperatorData>(WidgetOperatorData);
+			MyFloatPropertyHandle->GetValue(MutableWidgetOperatorData->MyFloat);
+		}
+	}
+}
+
+#undef LOCTEXT_NAMESPACE
diff --git a/Engine/Plugins/Experimental/MetasoundExperimental/Source/MetasoundExperimentalEditor/Private/MetasoundExampleNodeDetailsCustomization.h b/Engine/Plugins/Experimental/MetasoundExperimental/Source/MetasoundExperimentalEditor/Private/MetasoundExampleNodeDetailsCustomization.h
new file mode 100644
--- /dev/null
+++ b/Engine/Plugins/Experimental/MetasoundExperimental/Source/MetasoundExperimentalEditor/Private/MetasoundExampleNodeDetailsCustomization.h
@@ -0,0 +1,20 @@
+// Copyright Epic Games, Inc. All Rights Reserved.
+#pragma once
+
+#define LOCTEXT_NAMESPACE "MetasoundExperimentalEditor"
+
+#include "MetasoundNodeConfigurationCustomization.h"
+
+class FExampleWidgetNodeConfigurationCustomization : public Metasound::Editor::FMetaSoundNodeConfigurationDataDetails
+{
+public:
+	FExampleWidgetNodeConfigurationCustomization(TSharedPtr<IPropertyHandle> InStructProperty, TWeakObjectPtr<UMetasoundEditorGraphNode> InNode);
+
+	virtual void OnChildRowAdded(IDetailPropertyRow& ChildRow) override;
+
+private:
+	void OnChildPropertyChanged(const FPropertyChangedEvent& InPropertyChangedEvent);
+
+	FString StructPropertyPath;
+	TSharedPtr<IPropertyHandle> MyFloatPropertyHandle;
+};
diff --git a/Engine/Plugins/Experimental/MetasoundExperimental/Source/MetasoundExperimentalEditor/Private/MetasoundExperimentalEditorModule.cpp b/Engine/Plugins/Experimental/MetasoundExperimental/Source/MetasoundExperimentalEditor/Private/MetasoundExperimentalEditorModule.cpp
--- a/Engine/Plugins/Experimental/MetasoundExperimental/Source/MetasoundExperimentalEditor/Private/MetasoundExperimentalEditorModule.cpp
+++ b/Engine/Plugins/Experimental/MetasoundExperimental/Source/MetasoundExperimentalEditor/Private/MetasoundExperimentalEditorModule.cpp
@@ -2,7 +2,9 @@
 
 #define LOCTEXT_NAMESPACE "MetasoundExperimentalEditor"
 
-#include "Modules/ModuleInterface.h"
+#include "MetasoundEditorModule.h"
+#include "MetasoundExampleNodeConfiguration.h"
+#include "MetasoundExampleNodeDetailsCustomization.h"
 #include "Modules/ModuleManager.h"
 
 class FMetasoundExperimentalEditorModule : public IModuleInterface
@@ -12,18 +14,25 @@ public:
 	virtual void StartupModule() override;
 	virtual void ShutdownModule() override;
 	//~ End IModuleInterface API
-
-private:
 };
 
 void FMetasoundExperimentalEditorModule::StartupModule()
 {
+	Metasound::Editor::IMetasoundEditorModule& MetasoundEditorModule = FModuleManager::GetModuleChecked<Metasound::Editor::IMetasoundEditorModule>("MetasoundEditor");
+	MetasoundEditorModule.RegisterCustomNodeConfigurationDetailsCustomization(
+		FMetaSoundWidgetExampleNodeConfiguration::StaticStruct()->GetFName(),
+		[](TSharedPtr<IPropertyHandle> InStructProperty, TWeakObjectPtr<UMetasoundEditorGraphNode> InNode)
+		{
+			return MakeShareable(new FExampleWidgetNodeConfigurationCustomization(InStructProperty, InNode));
+		});
 }
 
 void FMetasoundExperimentalEditorModule::ShutdownModule()
 {
+	Metasound::Editor::IMetasoundEditorModule& MetasoundEditorModule = FModuleManager::GetModuleChecked<Metasound::Editor::IMetasoundEditorModule>("MetasoundEditor");
+	MetasoundEditorModule.UnregisterCustomNodeConfigurationDetailsCustomization(FMetaSoundWidgetExampleNodeConfiguration::StaticStruct()->GetFName());
 }
 
 #undef LOCTEXT_NAMESPACE
 
-IMPLEMENT_MODULE(FMetasoundExperimentalEditorModule, MetasoundExperimentalEditor)
+IMPLEMENT_MODULE(FMetasoundExperimentalEditorModule, MetasoundExperimentalEditor)
\ No newline at end of file
diff --git a/Engine/Plugins/Experimental/MetasoundExperimental/Source/MetasoundExperimentalRuntime/Private/MetasoundExampleNodeConfiguration.cpp b/Engine/Plugins/Experimental/MetasoundExperimental/Source/MetasoundExperimentalRuntime/Private/MetasoundExampleNodeConfiguration.cpp
--- a/Engine/Plugins/Experimental/MetasoundExperimental/Source/MetasoundExperimentalRuntime/Private/MetasoundExampleNodeConfiguration.cpp
+++ b/Engine/Plugins/Experimental/MetasoundExperimental/Source/MetasoundExperimentalRuntime/Private/MetasoundExampleNodeConfiguration.cpp
@@ -12,6 +12,7 @@
 #include "MetasoundExampleNodeConfiguration.h"
 #include "MetasoundExecutableOperator.h"
 #include "MetasoundParamHelper.h"
+#include "MetasoundPrimitives.h"
 #include "MetasoundTrigger.h"
 
 
@@ -94,12 +95,14 @@ namespace Metasound::Experimental
 					MoveTemp(OutputInterface)
 				};
 		}
+
+		METASOUND_PARAM(OutFloat, "Out", "Float output");
 	}
 
 	/** To send data from a FMetaSoundFrontendNodeConfiguration to an IOperator, it should
 	 * be encapsulated in the form of a IOperatorData.  
 	 *
-	 * The use of the TOperatorData provides some saftey mechanisms for downcasting node configurations.
+	 * The use of the TOperatorData provides some safety mechanisms for downcasting node configurations.
 	 */
 	class FExampleOperatorData : public TOperatorData<FExampleOperatorData>
 	{
@@ -281,6 +284,84 @@ namespace Metasound::Experimental
 
 	// The node extension must be registered along with the node. 
 	METASOUND_REGISTER_NODE_AND_CONFIGURATION(FExampleConfigurableNode, FMetaSoundExperimentalExampleNodeConfiguration);
+
+	// FMetaSoundWidgetExampleNodeConfiguration
+	const FLazyName FWidgetExampleOperatorData::OperatorDataTypeName = "ExperimentalWidgetExampleOperatorData";
+
+	class FWidgetExampleConfigurableOperator : public TExecutableOperator<FWidgetExampleConfigurableOperator>
+	{
+	public:
+		FWidgetExampleConfigurableOperator(const TSharedPtr<const FWidgetExampleOperatorData>& InOperatorData)
+			: OperatorData(InOperatorData)
+			, FloatOut(FFloatWriteRef::CreateNew())
+		{
+		}
+
+		virtual void BindOutputs(FOutputVertexInterfaceData& InOutVertexData) override
+		{
+			using namespace ExampleNodeConfigurationPrivate;
+			InOutVertexData.BindWriteVertex(METASOUND_GET_PARAM_NAME(OutFloat), FloatOut);
+		}
+
+		void Execute()
+		{
+			*FloatOut = OperatorData->MyFloat;
+		}
+		
+		static const FVertexInterface& GetVertexInterface()
+		{
+			using namespace ExampleNodeConfigurationPrivate;
+			static const FVertexInterface Interface(
+				FInputVertexInterface(),
+				FOutputVertexInterface(
+					TOutputDataVertex<float>(METASOUND_GET_PARAM_NAME_AND_METADATA(OutFloat))
+				)
+			);
+
+			return Interface;
+		}
+
+		static FNodeClassMetadata GetNodeInfo()
+		{
+			using namespace ExampleNodeConfigurationPrivate;
+			return FNodeClassMetadata
+			{
+				FNodeClassName{ "Experimental", "WidgetConfigurableOperator", "" },
+				1, // Major version
+				0, // Minor version
+				LOCTEXT("WidgetExampleConfigurableNodeName", "A Widget Configurable Node"),
+				LOCTEXT("WidgetExampleConfigurableNodeDescription", "A Node which shows how to make a configurable node with a custom details customization for yourself."),
+				TEXT("UE"), // Author
+				LOCTEXT("WidgetExampleConfigurablePromptIfMissing", "Enable the MetaSoundExperimental Plugin"), // Prompt if missing
+				GetVertexInterface(),
+				{}
+			};
+		}
+
+		static TUniquePtr<IOperator> CreateOperator(const FBuildOperatorParams& InParams, FBuildResults& OutResults)
+		{
+			using namespace ExampleNodeConfigurationPrivate;
+
+			if (const FWidgetExampleOperatorData* ExampleConfig = CastOperatorData<const FWidgetExampleOperatorData>(InParams.Node.GetOperatorData().Get()))
+			{
+				return MakeUnique<FWidgetExampleConfigurableOperator>(
+					StaticCastSharedPtr<const FWidgetExampleOperatorData>(InParams.Node.GetOperatorData())
+				);
+			}
+
+			return nullptr;
+		}
+
+	private:
+		// Contains configured float data
+		TSharedPtr<const FWidgetExampleOperatorData> OperatorData;
+
+		// Output
+		FFloatWriteRef FloatOut;
+	};
+
+	using FWidgetExampleConfigurableNode = TNodeFacade<FWidgetExampleConfigurableOperator>;
+	METASOUND_REGISTER_NODE_AND_CONFIGURATION(FWidgetExampleConfigurableNode, FMetaSoundWidgetExampleNodeConfiguration);
 }
 
 FMetaSoundExperimentalExampleNodeConfiguration::FMetaSoundExperimentalExampleNodeConfiguration()
@@ -303,4 +384,14 @@ TSharedPtr<const Metasound::IOperatorData> FMetaSoundExperimentalExampleNodeConf
 	return MakeShared<Metasound::Experimental::FExampleOperatorData>(String);
 }
 
+FMetaSoundWidgetExampleNodeConfiguration::FMetaSoundWidgetExampleNodeConfiguration()
+	: MyFloat(0.5f)
+	, OperatorData(MakeShared<Metasound::Experimental::FWidgetExampleOperatorData>(MyFloat))
+{
+}
+
+TSharedPtr<const Metasound::IOperatorData> FMetaSoundWidgetExampleNodeConfiguration::GetOperatorData() const
+{
+	return OperatorData;
+}
 #undef LOCTEXT_NAMESPACE // MetasoundExperimentalRuntime
diff --git a/Engine/Plugins/Experimental/MetasoundExperimental/Source/MetasoundExperimentalRuntime/Public/MetasoundExampleNodeConfiguration.h b/Engine/Plugins/Experimental/MetasoundExperimental/Source/MetasoundExperimentalRuntime/Public/MetasoundExampleNodeConfiguration.h
--- a/Engine/Plugins/Experimental/MetasoundExperimental/Source/MetasoundExperimentalRuntime/Public/MetasoundExampleNodeConfiguration.h
+++ b/Engine/Plugins/Experimental/MetasoundExperimental/Source/MetasoundExperimentalRuntime/Public/MetasoundExampleNodeConfiguration.h
@@ -7,6 +7,11 @@
 
 #include "MetasoundExampleNodeConfiguration.generated.h"
 
+namespace Metasound::Experimental
+{
+	class FWidgetExampleOperatorData;
+}
+
 USTRUCT()
 struct FMetaSoundExperimentalExampleNodeConfiguration : public FMetaSoundFrontendNodeConfiguration
 {
@@ -29,3 +34,35 @@ struct FMetaSoundExperimentalExampleNodeConfiguration : public FMetaSoundFronten
 	/* Pass data down to the operator. */
 	virtual TSharedPtr<const Metasound::IOperatorData> GetOperatorData() const override;
 };
+
+namespace Metasound::Experimental
+{
+	class FWidgetExampleOperatorData : public TOperatorData<FWidgetExampleOperatorData>
+	{
+	public:
+		static const FLazyName OperatorDataTypeName;
+
+		FWidgetExampleOperatorData(const float& InFloat)
+			: MyFloat(InFloat)
+		{
+		}
+
+		float MyFloat;
+	};
+}
+
+USTRUCT()
+struct FMetaSoundWidgetExampleNodeConfiguration : public FMetaSoundFrontendNodeConfiguration
+{
+	GENERATED_BODY()
+
+	FMetaSoundWidgetExampleNodeConfiguration();
+
+	UPROPERTY(EditAnywhere, Category = General)
+	float MyFloat;
+
+	TSharedPtr<const Metasound::IOperatorData> GetOperatorData() const override;
+
+private: 
+	TSharedPtr<Metasound::Experimental::FWidgetExampleOperatorData> OperatorData;
+};
diff --git a/Engine/Plugins/Runtime/Metasound/Source/MetasoundEditor/Private/MetasoundEditorModule.cpp b/Engine/Plugins/Runtime/Metasound/Source/MetasoundEditor/Private/MetasoundEditorModule.cpp
--- a/Engine/Plugins/Runtime/Metasound/Source/MetasoundEditor/Private/MetasoundEditorModule.cpp
+++ b/Engine/Plugins/Runtime/Metasound/Source/MetasoundEditor/Private/MetasoundEditorModule.cpp
@@ -674,6 +674,16 @@ namespace Metasound::Editor
 
 		}
 
+		virtual void RegisterCustomNodeConfigurationDetailsCustomization(FName InNodeConfigurationStructType, FCreateNodeConfigurationDetails InCreateDetailsFunc) override
+		{
+			CustomNodeConfigurationDetails.FindOrAdd(InNodeConfigurationStructType, InCreateDetailsFunc);
+		}
+
+		virtual void UnregisterCustomNodeConfigurationDetailsCustomization(FName InNodeConfigurationStructType) override
+		{
+			CustomNodeConfigurationDetails.Remove(InNodeConfigurationStructType);
+		}
+
 		virtual void RegisterGraphNodeVisualization(FName InNodeClassName, FOnCreateGraphNodeVisualizationWidget OnCreateGraphNodeVisualizationWidget) override
 		{
 			FGraphNodeVisualizationRegistry::Get().RegisterVisualization(InNodeClassName, OnCreateGraphNodeVisualizationWidget);
@@ -776,6 +786,11 @@ namespace Metasound::Editor
 			return nullptr;
 		}
 
+		virtual const FCreateNodeConfigurationDetails* FindCreateCustomNodeConfigurationDetailsCustomization(FName InNodeConfigurationStructType) const override
+		{
+			return CustomNodeConfigurationDetails.Find(InNodeConfigurationStructType);
+		}
+
 		virtual const TSubclassOf<UMetasoundEditorGraphMemberDefaultLiteral> FindDefaultLiteralClass(EMetasoundFrontendLiteralType InLiteralType) const override
 		{
 			return InputDefaultLiteralClassRegistry.FindRef(InLiteralType);
@@ -1106,6 +1121,7 @@ namespace Metasound::Editor
 		TSet<FName> CustomPinCategories;
 
 		TMap<UClass*, TUniquePtr<IMemberDefaultLiteralCustomizationFactory>> LiteralCustomizationFactories;
+		TMap<FName, FCreateNodeConfigurationDetails> CustomNodeConfigurationDetails;
 
 		TSharedPtr<FMetasoundGraphNodeFactory> GraphNodeFactory;
 		TSharedPtr<FGraphPanelPinConnectionFactory> GraphConnectionFactory;
diff --git a/Engine/Plugins/Runtime/Metasound/Source/MetasoundEditor/Private/MetasoundNodeConfigurationCustomization.cpp b/Engine/Plugins/Runtime/Metasound/Source/MetasoundEditor/Private/MetasoundNodeConfigurationCustomization.cpp
new file mode 100644
--- /dev/null
+++ b/Engine/Plugins/Runtime/Metasound/Source/MetasoundEditor/Private/MetasoundNodeConfigurationCustomization.cpp
@@ -0,0 +1,40 @@
+// Copyright Epic Games, Inc. All Rights Reserved.
+#pragma once
+
+#include "MetasoundNodeConfigurationCustomization.h"
+
+#include "MetasoundEditorGraphNode.h"
+#include "MetasoundFrontendDocumentBuilder.h"
+#include "MetasoundBuilderBase.h"
+#include "PropertyHandle.h"
+#include "IDetailPropertyRow.h"
+
+namespace Metasound::Editor
+{
+	void FMetaSoundNodeConfigurationDataDetails::OnChildRowAdded(IDetailPropertyRow& ChildRow)
+	{
+		FInstancedStructDataDetails::OnChildRowAdded(ChildRow);
+
+		TSharedPtr<IPropertyHandle> ChildHandle = ChildRow.GetPropertyHandle();
+		if (ChildHandle)
+		{
+			TDelegate<void(const FPropertyChangedEvent&)> OnValueChangedDelegate = TDelegate<void(const FPropertyChangedEvent&)>::CreateSP(this, &FMetaSoundNodeConfigurationDataDetails::OnChildPropertyChanged);
+			ChildHandle->SetOnPropertyValueChangedWithData(OnValueChangedDelegate);
+		}
+	}
+
+	void FMetaSoundNodeConfigurationDataDetails::OnChildPropertyChanged(const FPropertyChangedEvent& InPropertyChangedEvent)
+	{
+		if (GraphNode.IsValid())
+		{
+			// Don't update interface on interactive changes (ex. dragging slider) to avoid refresh spam
+			if (InPropertyChangedEvent.ChangeType != EPropertyChangeType::Interactive)
+			{
+				// Update the node config interface
+				FMetaSoundFrontendDocumentBuilder& DocBuilder = GraphNode->GetBuilderChecked().GetBuilder();
+				const FGuid& NodeID = GraphNode->GetNodeID();
+				DocBuilder.UpdateNodeInterfaceFromConfiguration(NodeID);
+			}
+		}
+	}
+} // namespace Metasound::Editor
diff --git a/Engine/Plugins/Runtime/Metasound/Source/MetasoundEditor/Private/MetasoundNodeDetailCustomization.cpp b/Engine/Plugins/Runtime/Metasound/Source/MetasoundEditor/Private/MetasoundNodeDetailCustomization.cpp
--- a/Engine/Plugins/Runtime/Metasound/Source/MetasoundEditor/Private/MetasoundNodeDetailCustomization.cpp
+++ b/Engine/Plugins/Runtime/Metasound/Source/MetasoundEditor/Private/MetasoundNodeDetailCustomization.cpp
@@ -9,6 +9,7 @@
 #include "HAL/PlatformApplicationMisc.h"
 #include "IDetailChildrenBuilder.h"
 #include "IDetailGroup.h"
+#include "InstancedStructDetails.h"
 #include "Internationalization/Text.h"
 #include "MetasoundDataReference.h"
 #include "MetasoundDataReferenceMacro.h"
@@ -20,6 +21,7 @@
 #include "MetasoundFrontendDataTypeRegistry.h"
 #include "MetasoundFrontendRegistries.h"
 #include "MetasoundEditorSettings.h"
+#include "MetasoundNodeConfigurationCustomization.h"
 #include "Modules/ModuleManager.h"
 #include "PropertyCustomizationHelpers.h"
 #include "PropertyEditorDelegates.h"
@@ -1062,31 +1064,6 @@ PRAGMA_ENABLE_DEPRECATION_WARNINGS
 				}
 			}
 		}
-		
-		void FMetaSoundNodeConfigurationDataDetails::OnChildRowAdded(IDetailPropertyRow& ChildRow)
-		{
-			FInstancedStructDataDetails::OnChildRowAdded(ChildRow);
-
-			TSharedPtr<IPropertyHandle> ChildHandle = ChildRow.GetPropertyHandle();
-			if (ChildHandle)
-			{
-				const FSimpleDelegate OnChildPropertyChangedDelegate = FSimpleDelegate::CreateSP(this, &FMetaSoundNodeConfigurationDataDetails::OnChildPropertyChanged);
-				ChildHandle->SetOnPropertyValueChanged(OnChildPropertyChangedDelegate);
-			}
-		}
-
-		void FMetaSoundNodeConfigurationDataDetails::OnChildPropertyChanged()
-		{
-			if (GraphNode.IsValid())
-			{
-				FMetaSoundFrontendDocumentBuilder& DocBuilder = GraphNode->GetBuilderChecked().GetBuilder();
-				const FGuid& NodeID = GraphNode->GetNodeID();
-				DocBuilder.UpdateNodeInterfaceFromConfiguration(NodeID);
-
-				FMetasoundFrontendDocumentModifyContext& ModifyContext = FGraphBuilder::GetOutermostMetaSoundChecked(*GraphNode).GetModifyContext();
-				ModifyContext.AddNodeIDModified(NodeID);
-			}
-		}
 
 		void FMetaSoundNodeDetailCustomization::CustomizeDetails(IDetailLayoutBuilder& InDetailLayout)
 		{
@@ -1120,9 +1097,30 @@ PRAGMA_ENABLE_DEPRECATION_WARNINGS
 
 						// Use a custom details builder to add custom child property update behavior
 						// and hide the struct picker since users should only be able edit child values, not the struct type 
-						if (ConfigurationHandle->IsValidHandle())
+						if (ConfigurationHandle && ConfigurationHandle->IsValidHandle())
 						{
-							TSharedRef<IDetailCustomNodeBuilder> DataDetailsBuilder = MakeShared<FMetaSoundNodeConfigurationDataDetails>(ConfigurationHandle, TWeakObjectPtr<UMetasoundEditorGraphNode>(Node));
+							const IMetasoundEditorModule& MetaSoundEditorModule = FModuleManager::GetModuleChecked<IMetasoundEditorModule>("MetaSoundEditor");
+
+							// Get underlying struct name
+							FName StructName;
+							ConfigurationHandle->EnumerateRawData([&StructName](void* RawData, const int32 DataIndex, const int32 /*NumDatas*/)
+							{
+								if (FInstancedStruct* InstancedStruct = static_cast<FInstancedStruct*>(RawData))
+								{
+									StructName = InstancedStruct->GetScriptStruct()->GetFName();
+								}
+								return true;
+							});
+
+							const FCreateNodeConfigurationDetails* CreateDetailsFunc = !StructName.IsNone() ? 
+								MetaSoundEditorModule.FindCreateCustomNodeConfigurationDetailsCustomization(StructName) : nullptr;
+
+							TSharedRef<IDetailCustomNodeBuilder> DataDetailsBuilder = CreateDetailsFunc ? 
+								// User custom registered details 
+								(*CreateDetailsFunc)(ConfigurationHandle, TWeakObjectPtr<UMetasoundEditorGraphNode>(Node)) : 
+								// Default details
+								MakeShared<FMetaSoundNodeConfigurationDataDetails>(ConfigurationHandle, TWeakObjectPtr<UMetasoundEditorGraphNode>(Node));
+
 							InDetailLayout.EditCategory("General").AddCustomBuilder(DataDetailsBuilder);
 						}
 					}
diff --git a/Engine/Plugins/Runtime/Metasound/Source/MetasoundEditor/Private/MetasoundNodeDetailCustomization.h b/Engine/Plugins/Runtime/Metasound/Source/MetasoundEditor/Private/MetasoundNodeDetailCustomization.h
--- a/Engine/Plugins/Runtime/Metasound/Source/MetasoundEditor/Private/MetasoundNodeDetailCustomization.h
+++ b/Engine/Plugins/Runtime/Metasound/Source/MetasoundEditor/Private/MetasoundNodeDetailCustomization.h
@@ -11,7 +11,6 @@
 #include "IDetailCustomization.h"
 #include "IDetailPropertyExtensionHandler.h"
 #include "IDetailPropertyRow.h"
-#include "InstancedStructDetails.h"
 #include "IPropertyTypeCustomization.h"
 #include "Layout/Visibility.h"
 #include "MetasoundAssetBase.h"
@@ -319,23 +318,6 @@ namespace Metasound
 			FName ArrayTypeName;
 		};
 
-		class FMetaSoundNodeConfigurationDataDetails : public FInstancedStructDataDetails
-		{
-		public:
-			FMetaSoundNodeConfigurationDataDetails(TSharedPtr<IPropertyHandle> InStructProperty, TWeakObjectPtr<UMetasoundEditorGraphNode> InNode)
-				: FInstancedStructDataDetails(InStructProperty)
-				, GraphNode(InNode)
-			{
-			}
-
-			virtual void OnChildRowAdded(IDetailPropertyRow& ChildRow) override;
-		
-		private: 
-			void OnChildPropertyChanged();
-
-			TWeakObjectPtr<UMetasoundEditorGraphNode> GraphNode;
-		};
-
 		class FMetaSoundNodeDetailCustomization : public IDetailCustomization
 		{
 			// IDetailCustomization interface
diff --git a/Engine/Plugins/Runtime/Metasound/Source/MetasoundEditor/Public/MetasoundEditorGraphNode.h b/Engine/Plugins/Runtime/Metasound/Source/MetasoundEditor/Public/MetasoundEditorGraphNode.h
--- a/Engine/Plugins/Runtime/Metasound/Source/MetasoundEditor/Public/MetasoundEditorGraphNode.h
+++ b/Engine/Plugins/Runtime/Metasound/Source/MetasoundEditor/Public/MetasoundEditorGraphNode.h
@@ -132,7 +132,7 @@ public:
 	virtual FName GetCornerIcon() const { return FName(); }
 	virtual bool CanAddInputPin() const { return false; }
 
-	UMetaSoundBuilderBase& GetBuilderChecked() const;
+	METASOUNDEDITOR_API UMetaSoundBuilderBase& GetBuilderChecked() const;
 
 	// Returns document's cached Frontend Node Class (as defined by the
 	// document's dependency array). If node or class is not found on
diff --git a/Engine/Plugins/Runtime/Metasound/Source/MetasoundEditor/Public/MetasoundEditorModule.h b/Engine/Plugins/Runtime/Metasound/Source/MetasoundEditor/Public/MetasoundEditorModule.h
--- a/Engine/Plugins/Runtime/Metasound/Source/MetasoundEditor/Public/MetasoundEditorModule.h
+++ b/Engine/Plugins/Runtime/Metasound/Source/MetasoundEditor/Public/MetasoundEditorModule.h
@@ -50,8 +50,12 @@ namespace Metasound
 			const FSlateColor& GetPageExecutingColor();
 		} // namespace Style
 
+		class FMetaSoundNodeConfigurationDataDetails;
+
 		using EAssetScanStatus = Metasound::Engine::EAssetScanStatus;
 		using EAssetPrimeStatus = Metasound::Engine::ENodeClassRegistryPrimeStatus;
+		using FCreateNodeConfigurationDetails = TFunction<TSharedRef<FMetaSoundNodeConfigurationDataDetails>(
+			TSharedPtr<IPropertyHandle> InStructProperty, TWeakObjectPtr<UMetasoundEditorGraphNode> InNode)>;
 
 		struct FGraphPinParams
 		{
@@ -94,6 +98,9 @@ namespace Metasound
 
 			virtual TUniquePtr<FMetasoundDefaultLiteralCustomizationBase> CreateMemberDefaultLiteralCustomization(UClass& InClass, IDetailCategoryBuilder& DefaultCategoryBuilder) const = 0;
 
+			UE_EXPERIMENTAL(5.6, "Node configuration custom details customizations are experimental.")
+			virtual const FCreateNodeConfigurationDetails* FindCreateCustomNodeConfigurationDetailsCustomization(FName InNodeConfigurationStructType) const = 0;
+
 			virtual const TSubclassOf<UMetasoundEditorGraphMemberDefaultLiteral> FindDefaultLiteralClass(EMetasoundFrontendLiteralType InLiteralType) const = 0;
 
 			virtual const FEdGraphPinType* FindPinType(FName InDataTypeName) const = 0;
@@ -108,6 +115,16 @@ namespace Metasound
 			
 			virtual void RegisterCustomPinType(FName InDataTypeName, const FGraphPinParams& Params) = 0;
 
+			/** Register a custom details customization for a specific substruct of FMetaSoundFrontendNodeConfiguration 
+				InNodeConfigurationStructType - struct type name - use MyStruct::StaticStruct()->GetFName()
+				InCreateDetailsFunc - function that creates the specific details customization */
+			UE_EXPERIMENTAL(5.6, "Node configuration custom details customizations are experimental.")
+			virtual void RegisterCustomNodeConfigurationDetailsCustomization(FName InNodeConfigurationStructType, FCreateNodeConfigurationDetails InCreateDetailsFunc) = 0;
+
+			// Unregister a custom details customization for a node configuration struct type, using the same struct type name as RegisterCustomNodeConfigurationDetailsCustomization
+			UE_EXPERIMENTAL(5.6, "Node configuration custom details customizations are experimental.")
+			virtual void UnregisterCustomNodeConfigurationDetailsCustomization(FName InNodeConfigurationStructType) = 0;
+
 			// For the given node class, register a delegate that can be used for creating in-graph node visualizations.
 			virtual void RegisterGraphNodeVisualization(FName InNodeClassName, FOnCreateGraphNodeVisualizationWidget OnCreateGraphNodeVisualizationWidget) = 0;
 
diff --git a/Engine/Plugins/Runtime/Metasound/Source/MetasoundEditor/Public/MetasoundNodeConfigurationCustomization.h b/Engine/Plugins/Runtime/Metasound/Source/MetasoundEditor/Public/MetasoundNodeConfigurationCustomization.h
new file mode 100644
--- /dev/null
+++ b/Engine/Plugins/Runtime/Metasound/Source/MetasoundEditor/Public/MetasoundNodeConfigurationCustomization.h
@@ -0,0 +1,31 @@
+// Copyright Epic Games, Inc. All Rights Reserved.
+#pragma once
+
+#include "InstancedStructDetails.h"
+
+// Forward Declarations
+struct FPropertyChangedEvent;
+class FInstancedStructDataDetails;
+class IDetailPropertyRow;
+class IPropertyHandle;
+class UMetasoundEditorGraphNode;
+
+namespace Metasound::Editor
+{
+	class METASOUNDEDITOR_API FMetaSoundNodeConfigurationDataDetails : public FInstancedStructDataDetails
+	{
+	public:
+		FMetaSoundNodeConfigurationDataDetails(TSharedPtr<IPropertyHandle> InStructProperty, TWeakObjectPtr<UMetasoundEditorGraphNode> InNode)
+			: FInstancedStructDataDetails(InStructProperty)
+			, GraphNode(InNode)
+		{
+		}
+
+		virtual void OnChildRowAdded(IDetailPropertyRow& ChildRow) override;
+
+	protected:
+		void OnChildPropertyChanged(const FPropertyChangedEvent& InPropertyChangedEvent);
+
+		TWeakObjectPtr<UMetasoundEditorGraphNode> GraphNode;
+	};
+} // namespace Metasound::Editor
diff --git a/Engine/Plugins/Runtime/Metasound/Source/MetasoundFrontend/Private/MetasoundFrontendDocument.cpp b/Engine/Plugins/Runtime/Metasound/Source/MetasoundFrontend/Private/MetasoundFrontendDocument.cpp
--- a/Engine/Plugins/Runtime/Metasound/Source/MetasoundFrontend/Private/MetasoundFrontendDocument.cpp
+++ b/Engine/Plugins/Runtime/Metasound/Source/MetasoundFrontend/Private/MetasoundFrontendDocument.cpp
@@ -487,15 +487,16 @@ FMetasoundFrontendNodeInterface::FMetasoundFrontendNodeInterface(const FMetasoun
 	}
 }
 
-void FMetasoundFrontendNodeInterface::Update(const FMetasoundFrontendClassInterface& InClassInterface)
+bool FMetasoundFrontendNodeInterface::Update(const FMetasoundFrontendClassInterface& InClassInterface)
 {
 	auto NoOp = [](const FMetasoundFrontendVertex&)->void {};
 
-	Update(InClassInterface, NoOp, NoOp);
+	return Update(InClassInterface, NoOp, NoOp);
 }
 
-void FMetasoundFrontendNodeInterface::Update(const FMetasoundFrontendClassInterface& InClassInterface, TFunctionRef<void(const FMetasoundFrontendVertex&)> OnPreRemoveInput, TFunctionRef<void(const FMetasoundFrontendVertex&)> OnPreRemoveOutput)
+bool FMetasoundFrontendNodeInterface::Update(const FMetasoundFrontendClassInterface& InClassInterface, TFunctionRef<void(const FMetasoundFrontendVertex&)> OnPreRemoveInput, TFunctionRef<void(const FMetasoundFrontendVertex&)> OnPreRemoveOutput)
 {
+	bool bInterfaceUpdated = false;
 	TArray<const FMetasoundFrontendVertex*, TInlineAllocator<32>> UnmatchedVertices;
 
 	struct FFindMatchingVertex
@@ -520,6 +521,7 @@ void FMetasoundFrontendNodeInterface::Update(const FMetasoundFrontendClassInterf
 		{
 			// Add class input to node inputs
 			Inputs.Add(ClassInput);
+			bInterfaceUpdated |= true;
 		}
 	}
 
@@ -529,6 +531,7 @@ void FMetasoundFrontendNodeInterface::Update(const FMetasoundFrontendClassInterf
 		// Allow outside systems to react before removing the unmatched node inputs
 		OnPreRemoveInput(*UnmatchedInput);
 		Inputs.RemoveSingleSwap(*UnmatchedInput);
+		bInterfaceUpdated |= true;
 	}
 
 	// Update node outputs
@@ -545,6 +548,7 @@ void FMetasoundFrontendNodeInterface::Update(const FMetasoundFrontendClassInterf
 		{
 			// Add class input to node inputs
 			Outputs.Add(ClassOutput);
+			bInterfaceUpdated |= true;
 		}
 	}
 
@@ -554,7 +558,10 @@ void FMetasoundFrontendNodeInterface::Update(const FMetasoundFrontendClassInterf
 		// Allow outside systems to react before removing the unmatched node inputs
 		OnPreRemoveOutput(*UnmatchedOutput);
 		Outputs.RemoveSingleSwap(*UnmatchedOutput);
+		bInterfaceUpdated |= true;
 	}
+	
+	return bInterfaceUpdated;
 }
 
 TInstancedStruct<FMetasoundFrontendClassInterface> FMetaSoundFrontendNodeConfiguration::OverrideDefaultInterface(const FMetasoundFrontendClass& InNodeClass) const
diff --git a/Engine/Plugins/Runtime/Metasound/Source/MetasoundFrontend/Private/MetasoundFrontendDocumentBuilder.cpp b/Engine/Plugins/Runtime/Metasound/Source/MetasoundFrontend/Private/MetasoundFrontendDocumentBuilder.cpp
--- a/Engine/Plugins/Runtime/Metasound/Source/MetasoundFrontend/Private/MetasoundFrontendDocumentBuilder.cpp
+++ b/Engine/Plugins/Runtime/Metasound/Source/MetasoundFrontend/Private/MetasoundFrontendDocumentBuilder.cpp
@@ -2898,7 +2898,19 @@ const FMetasoundFrontendNode* FMetaSoundFrontendDocumentBuilder::FindNode(const
 	return NodeCache.FindNode(InNodeID);
 }
 
-const TInstancedStruct<FMetaSoundFrontendNodeConfiguration> FMetaSoundFrontendDocumentBuilder::FindNodeConfiguration(const FGuid& InNodeID, const FGuid* InPageID) const
+const TConstStructView<FMetaSoundFrontendNodeConfiguration> FMetaSoundFrontendDocumentBuilder::FindNodeConfiguration(const FGuid& InNodeID, const FGuid* InPageID) const
+{
+	using namespace Metasound::Frontend;
+	const FGuid& PageID = InPageID ? *InPageID : BuildPageID;
+	const IDocumentGraphNodeCache& NodeCache = DocumentCache->GetNodeCache(PageID);
+	if (const FMetasoundFrontendNode* Node = NodeCache.FindNode(InNodeID))
+	{
+		return Node->Configuration;
+	}
+	return TConstStructView<FMetaSoundFrontendNodeConfiguration>();
+}
+
+TInstancedStruct<FMetaSoundFrontendNodeConfiguration> FMetaSoundFrontendDocumentBuilder::FindNodeConfiguration(const FGuid& InNodeID, const FGuid* InPageID) 
 {
 	using namespace Metasound::Frontend;
 	const FGuid& PageID = InPageID ? *InPageID : BuildPageID;
@@ -6335,7 +6347,7 @@ bool FMetaSoundFrontendDocumentBuilder::UpdateNodeInterfaceFromConfiguration(con
 			check(Class);
 
 			Node->ClassInterfaceOverride = ConfigurationPtr->OverrideDefaultInterface(*Class);
-			const FMetasoundFrontendClassInterface & ClassInterface = Class->GetInterfaceForNode(*Node);
+			const FMetasoundFrontendClassInterface& ClassInterface = Class->GetInterfaceForNode(*Node);
 
 			auto DisconnectInput = [this, &InNodeID, &InPageID](const FMetasoundFrontendVertex& InNodeInput)
 			{
@@ -6346,8 +6358,15 @@ bool FMetaSoundFrontendDocumentBuilder::UpdateNodeInterfaceFromConfiguration(con
 			{
 				this->RemoveEdgesFromNodeOutput(InNodeID, InNodeOutput.VertexID, InPageID);
 			};
-
+#if WITH_EDITORONLY_DATA
+			const bool bInterfaceUpdated = Node->Interface.Update(ClassInterface, DisconnectInput, DisconnectOutput);
+			if (bInterfaceUpdated)
+			{
+				GetDocumentChecked().Metadata.ModifyContext.AddNodeIDModified(InNodeID);
+			}
+#else
 			Node->Interface.Update(ClassInterface, DisconnectInput, DisconnectOutput);
+#endif
 		}
 
 		return true;
diff --git a/Engine/Plugins/Runtime/Metasound/Source/MetasoundFrontend/Public/MetasoundFrontendDocument.h b/Engine/Plugins/Runtime/Metasound/Source/MetasoundFrontend/Public/MetasoundFrontendDocument.h
--- a/Engine/Plugins/Runtime/Metasound/Source/MetasoundFrontend/Public/MetasoundFrontendDocument.h
+++ b/Engine/Plugins/Runtime/Metasound/Source/MetasoundFrontend/Public/MetasoundFrontendDocument.h
@@ -550,11 +550,13 @@ struct FMetasoundFrontendNodeInterface
 	// Create a node interface which satisfies an existing class interface.
 	UE_API FMetasoundFrontendNodeInterface(const FMetasoundFrontendClassInterface& InClassInterface);
 
+	// Update the current node interface with the given class interface
+	// Return true if interface update resulted in interface changes
 	UE_INTERNAL
-	UE_API void Update(const FMetasoundFrontendClassInterface& InClassInterface);
+	UE_API bool Update(const FMetasoundFrontendClassInterface& InClassInterface);
 
 	UE_INTERNAL
-	UE_API void Update(const FMetasoundFrontendClassInterface& InClassInterface, TFunctionRef<void(const FMetasoundFrontendVertex&)> OnPreRemoveInput, TFunctionRef<void(const FMetasoundFrontendVertex&)> OnPreRemoveOutput);
+	UE_API bool Update(const FMetasoundFrontendClassInterface& InClassInterface, TFunctionRef<void(const FMetasoundFrontendVertex&)> OnPreRemoveInput, TFunctionRef<void(const FMetasoundFrontendVertex&)> OnPreRemoveOutput);
 
 	// Input vertices to node.
 	UPROPERTY()
@@ -569,6 +571,31 @@ struct FMetasoundFrontendNodeInterface
 	TArray<FMetasoundFrontendVertex> Environment;
 };
 
+/** 
+* Struct for configuring a node. You can inherit from this 
+* and include data passed to the operator
+* and/or data used to determine an override of the node's interface.
+* 
+* In order for node configuration data to be editable in the details panel, 
+* UProperties on your substruct should be marked with EditAnywhere.
+* 
+* Optional custom details customizations can be registered via 
+* IMetasoundEditorModule::RegisterCustomNodeConfigurationDetailsCustomization
+* 
+* Example: 
+USTRUCT()
+struct FMetaSoundExperimentalExampleNodeConfiguration : public FMetaSoundFrontendNodeConfiguration
+{
+	GENERATED_BODY()
+	
+	UPROPERTY(EditAnywhere, Category = General, meta = (ClampMin = "1", ClampMax = "1000"))
+	uint32 NumInputs;
+
+	virtual TInstancedStruct<FMetasoundFrontendClassInterface> OverrideDefaultInterface(const FMetasoundFrontendClass& InNodeClass) const;
+
+	virtual TSharedPtr<const Metasound::IOperatorData> GetOperatorData() const override;
+}
+*/
 USTRUCT()
 struct FMetaSoundFrontendNodeConfiguration
 {
@@ -580,6 +607,7 @@ struct FMetaSoundFrontendNodeConfiguration
 	UE_API virtual TInstancedStruct<FMetasoundFrontendClassInterface> OverrideDefaultInterface(const FMetasoundFrontendClass& InNodeClass) const;
 
 	/** Provide any data needed by IOperators instantiated from this node. */
+	UE_EXPERIMENTAL(5.6, "Node operator data is still under development")
 	UE_API virtual TSharedPtr<const Metasound::IOperatorData> GetOperatorData() const;
 };
 
diff --git a/Engine/Plugins/Runtime/Metasound/Source/MetasoundFrontend/Public/MetasoundFrontendDocumentBuilder.h b/Engine/Plugins/Runtime/Metasound/Source/MetasoundFrontend/Public/MetasoundFrontendDocumentBuilder.h
--- a/Engine/Plugins/Runtime/Metasound/Source/MetasoundFrontend/Public/MetasoundFrontendDocumentBuilder.h
+++ b/Engine/Plugins/Runtime/Metasound/Source/MetasoundFrontend/Public/MetasoundFrontendDocumentBuilder.h
@@ -276,7 +276,12 @@ public:
 	UE_API const FMetasoundFrontendGraph& FindConstBuildGraphChecked() const;
 
 	UE_API const FMetasoundFrontendNode* FindNode(const FGuid& InNodeID, const FGuid* InPageID = nullptr) const;
-	UE_API const TInstancedStruct<FMetaSoundFrontendNodeConfiguration> FindNodeConfiguration(const FGuid& InNodeID, const FGuid* InPageID = nullptr) const;
+
+	UE_API const TConstStructView<FMetaSoundFrontendNodeConfiguration> FindNodeConfiguration(const FGuid& InNodeID, const FGuid* InPageID = nullptr) const;
+
+	UE_EXPERIMENTAL(5.6, "Non const builder access to node configuration is experimental.")
+	UE_API TInstancedStruct<FMetaSoundFrontendNodeConfiguration> FindNodeConfiguration(const FGuid& InNodeID, const FGuid* InPageID = nullptr);
+
 	// Return the node's index in the document's specified paged graph's node array 
 	UE_API const int32* FindNodeIndex(const FGuid& InNodeID, const FGuid* InPageID = nullptr) const;
 
diff --git a/Engine/Plugins/Runtime/Metasound/Source/MetasoundGraphCore/Public/MetasoundNodeInterface.h b/Engine/Plugins/Runtime/Metasound/Source/MetasoundGraphCore/Public/MetasoundNodeInterface.h
--- a/Engine/Plugins/Runtime/Metasound/Source/MetasoundGraphCore/Public/MetasoundNodeInterface.h
+++ b/Engine/Plugins/Runtime/Metasound/Source/MetasoundGraphCore/Public/MetasoundNodeInterface.h
@@ -196,8 +196,10 @@ namespace Metasound
 
 			/** Returns the configuration for the node. */
 #if UE_METASOUND_PURE_VIRTUAL_GET_OPERATOR_DATA
+			UE_EXPERIMENTAL(5.6, "Node operator data is still under development")
 			virtual TSharedPtr<const IOperatorData> GetOperatorData() const = 0;
 #else 
+			UE_EXPERIMENTAL(5.6, "Node operator data is still under development")
 			virtual METASOUNDGRAPHCORE_API TSharedPtr<const IOperatorData> GetOperatorData() const;
 #endif
 
