MetaSound experimental node loctext fix part 3

#rnx

diff --git a/Engine/Plugins/Experimental/MetasoundExperimental/Source/MetasoundExperimentalEditor/Private/MetasoundExampleNodeDetailsCustomization.h b/Engine/Plugins/Experimental/MetasoundExperimental/Source/MetasoundExperimentalEditor/Private/MetasoundExampleNodeDetailsCustomization.h
--- a/Engine/Plugins/Experimental/MetasoundExperimental/Source/MetasoundExperimentalEditor/Private/MetasoundExampleNodeDetailsCustomization.h
+++ b/Engine/Plugins/Experimental/MetasoundExperimental/Source/MetasoundExperimentalEditor/Private/MetasoundExampleNodeDetailsCustomization.h
@@ -3,8 +3,6 @@
 
 #include "MetasoundNodeConfigurationCustomization.h"
 
-#define LOCTEXT_NAMESPACE "MetasoundExperimentalEditor"
-
 class FExampleWidgetNodeConfigurationCustomization : public Metasound::Editor::FMetaSoundNodeConfigurationDataDetails
 {
 public:
