Fixed Audio::FSimpleLinearAllocator to go via FMemory::Malloc instead of calling GMalloc directly as it bypasses proper LLM tracking

diff --git a/Engine/Plugins/Experimental/MetasoundExperimental/Source/AudioExperimentalRuntime/Public/SimpleAlloc/SimpleAllocBase.h b/Engine/Plugins/Experimental/MetasoundExperimental/Source/AudioExperimentalRuntime/Public/SimpleAlloc/SimpleAllocBase.h
--- a/Engine/Plugins/Experimental/MetasoundExperimental/Source/AudioExperimentalRuntime/Public/SimpleAlloc/SimpleAllocBase.h
+++ b/Engine/Plugins/Experimental/MetasoundExperimental/Source/AudioExperimentalRuntime/Public/SimpleAlloc/SimpleAllocBase.h
@@ -15,35 +15,9 @@ namespace Audio
 		FSimpleAllocBase() = default;
 		virtual ~FSimpleAllocBase() = default;
 	
-		[[nodiscard]] virtual void* Malloc(const SIZE_T, const uint32 InAlignment=DEFAULT_ALIGNMENT) = 0;
-		virtual void Free(void*) {};
+		[[nodiscard]] virtual void* Malloc(const SIZE_T, const uint32 InAlignment = DEFAULT_ALIGNMENT) = 0;
+		virtual void Free(void*) = 0;
 		virtual uint32 GetCurrentLifetime() const { return 0; };
 		virtual void Reset() {};
 	};
-
-	namespace SimpleAllocBasePrivate
-	{
-		// This is what's defined in MemoryBase, so follow it here. 
-		FORCEINLINE static uint32 GetDefaultSizeToAlignment(const uint32 InSize)
-		{
-			check(InSize > 0);
-			return InSize >= 16 ? 16 : 8;
-		};
-
-		// When given an alignment and offset return a new rounded offset that would honor the alignment
-		FORCEINLINE static uint32 RoundUpToAlignment(const uint32 InOffset, const uint32 InAlignment)
-		{
-			check(InAlignment > 0);
-			check(FMath::IsPowerOfTwo(InAlignment));
-			const uint32 Mod = InOffset & (InAlignment-1);
-			const uint32 RoundUp = Mod > 0 ? InAlignment - Mod : 0;
-			return InOffset + RoundUp;
-		}
-
-		FORCEINLINE static bool IsAligned(const void* InPtr, const uint32 InAlignment)
-		{
-			check(FMath::IsPowerOfTwo(InAlignment));
-			return (reinterpret_cast<UPTRINT>(InPtr) & (InAlignment - 1)) == 0;
-		}
-	} //namespace FSimpleAllocBasePrivate
 }
\ No newline at end of file
diff --git a/Engine/Plugins/Experimental/MetasoundExperimental/Source/AudioExperimentalRuntime/Public/SimpleAlloc/SimpleLinearAllocator.h b/Engine/Plugins/Experimental/MetasoundExperimental/Source/AudioExperimentalRuntime/Public/SimpleAlloc/SimpleLinearAllocator.h
--- a/Engine/Plugins/Experimental/MetasoundExperimental/Source/AudioExperimentalRuntime/Public/SimpleAlloc/SimpleLinearAllocator.h
+++ b/Engine/Plugins/Experimental/MetasoundExperimental/Source/AudioExperimentalRuntime/Public/SimpleAlloc/SimpleLinearAllocator.h
@@ -4,6 +4,7 @@
 
 #include "SimpleAllocBase.h"
 #include "Containers/ArrayView.h"
+#include "HAL/MallocBinnedCommon.h"
 
 namespace Audio
 {
@@ -13,15 +14,15 @@ namespace Audio
 		explicit FSimpleLinearAllocator(const TArrayView<uint8> InPage)
 			: Page(InPage)
 		{
-			check(SimpleAllocBasePrivate::IsAligned(InPage.GetData(),SimpleAllocBasePrivate::GetDefaultSizeToAlignment(InPage.Num())));
+			check(IsAligned(InPage.GetData(), UE_MBC_MIN_SMALL_POOL_ALIGNMENT));
 		}
 	
-		virtual void* Malloc(const SIZE_T InSizeBytes, const uint32 InAlignment=DEFAULT_ALIGNMENT) override
+		virtual void* Malloc(const SIZE_T InSizeBytes, const uint32 InAlignment = DEFAULT_ALIGNMENT) override
 		{
-			const uint32 Alignment = InAlignment == DEFAULT_ALIGNMENT ? SimpleAllocBasePrivate::GetDefaultSizeToAlignment(InSizeBytes) : InAlignment;
-			const uint32 AlignedNewTop = SimpleAllocBasePrivate::RoundUpToAlignment(Top, Alignment);
+			const uint32 Alignment = FMath::Max(InAlignment, (uint32)UE_MBC_MIN_SMALL_POOL_ALIGNMENT);
+			const uint32 AlignedNewTop = Align(Top, Alignment);
 		
-			if (AlignedNewTop+InSizeBytes < Page.Num())
+			if (AlignedNewTop + InSizeBytes < Page.Num())
 			{
 				Top = AlignedNewTop;
 				void* Ret = Page.GetData() + Top;
@@ -50,15 +51,12 @@ namespace Audio
 	class FSimpleLinearAllocatorFromHeap final : public FSimpleLinearAllocator
 	{
 	public:
-		explicit FSimpleLinearAllocatorFromHeap(const SIZE_T InPageSize, const uint32 InPageAlignment = DEFAULT_ALIGNMENT, FMalloc& InHeap = *UE::Private::GMalloc)
-			: FSimpleLinearAllocator(MakeArrayView(static_cast<uint8*>(InHeap.Malloc(InPageSize,InPageAlignment)), IntCastChecked<int32>(InPageSize)))
-			, Heap(InHeap)
+		explicit FSimpleLinearAllocatorFromHeap(const SIZE_T InPageSize, const uint32 InPageAlignment = DEFAULT_ALIGNMENT)
+			: FSimpleLinearAllocator(MakeArrayView(static_cast<uint8*>(FMemory::Malloc(InPageSize,InPageAlignment)), IntCastChecked<int32>(InPageSize)))
 		{}
 		virtual ~FSimpleLinearAllocatorFromHeap() override
 		{
-			Heap.Free(Page.GetData());
+			FMemory::Free(Page.GetData());
 		}
-	private:
-		FMalloc& Heap;
 	};
 }
