Fixed compilation error in SimpleLinearAllocator due to a missing pure virtual function override

diff --git a/Engine/Plugins/Experimental/MetasoundExperimental/Source/AudioExperimentalRuntime/Public/SimpleAlloc/SimpleLinearAllocator.h b/Engine/Plugins/Experimental/MetasoundExperimental/Source/AudioExperimentalRuntime/Public/SimpleAlloc/SimpleLinearAllocator.h
--- a/Engine/Plugins/Experimental/MetasoundExperimental/Source/AudioExperimentalRuntime/Public/SimpleAlloc/SimpleLinearAllocator.h
+++ b/Engine/Plugins/Experimental/MetasoundExperimental/Source/AudioExperimentalRuntime/Public/SimpleAlloc/SimpleLinearAllocator.h
@@ -31,6 +31,9 @@ namespace Audio
 			}
 			return nullptr;
 		}
+
+		virtual void Free(void*) override {}
+
 		virtual uint32 GetCurrentLifetime() const override
 		{
 			return CurrentLifetime;
