Fixed AudioUnitTests

diff --git a/Engine/Plugins/Experimental/MetasoundExperimental/Source/AudioExperimentalRuntime/Public/SimpleAlloc/SimpleAllocBase.h b/Engine/Plugins/Experimental/MetasoundExperimental/Source/AudioExperimentalRuntime/Public/SimpleAlloc/SimpleAllocBase.h
--- a/Engine/Plugins/Experimental/MetasoundExperimental/Source/AudioExperimentalRuntime/Public/SimpleAlloc/SimpleAllocBase.h
+++ b/Engine/Plugins/Experimental/MetasoundExperimental/Source/AudioExperimentalRuntime/Public/SimpleAlloc/SimpleAllocBase.h
@@ -20,4 +20,14 @@ namespace Audio
 		virtual uint32 GetCurrentLifetime() const { return 0; };
 		virtual void Reset() {};
 	};
+
+	namespace SimpleAllocBasePrivate
+	{
+		// This is what's defined in MemoryBase, so follow it here. 
+		FORCEINLINE static uint32 GetDefaultSizeToAlignment(const uint32 InSize)
+		{
+			check(InSize > 0);
+			return InSize >= 16 ? 16 : 8;
+		};
+	}
 }
\ No newline at end of file
diff --git a/Engine/Plugins/Experimental/MetasoundExperimental/Source/AudioExperimentalRuntime/Public/SimpleAlloc/SimpleLinearAllocator.h b/Engine/Plugins/Experimental/MetasoundExperimental/Source/AudioExperimentalRuntime/Public/SimpleAlloc/SimpleLinearAllocator.h
--- a/Engine/Plugins/Experimental/MetasoundExperimental/Source/AudioExperimentalRuntime/Public/SimpleAlloc/SimpleLinearAllocator.h
+++ b/Engine/Plugins/Experimental/MetasoundExperimental/Source/AudioExperimentalRuntime/Public/SimpleAlloc/SimpleLinearAllocator.h
@@ -4,7 +4,6 @@
 
 #include "SimpleAllocBase.h"
 #include "Containers/ArrayView.h"
-#include "HAL/MallocBinnedCommon.h"
 
 namespace Audio
 {
@@ -14,12 +13,12 @@ namespace Audio
 		explicit FSimpleLinearAllocator(const TArrayView<uint8> InPage)
 			: Page(InPage)
 		{
-			check(IsAligned(InPage.GetData(), UE_MBC_MIN_SMALL_POOL_ALIGNMENT));
+			check(IsAligned(InPage.GetData(), SimpleAllocBasePrivate::GetDefaultSizeToAlignment(InPage.Num())));
 		}
 	
 		virtual void* Malloc(const SIZE_T InSizeBytes, const uint32 InAlignment = DEFAULT_ALIGNMENT) override
 		{
-			const uint32 Alignment = FMath::Max(InAlignment, (uint32)UE_MBC_MIN_SMALL_POOL_ALIGNMENT);
+			const uint32 Alignment = InAlignment == DEFAULT_ALIGNMENT ? SimpleAllocBasePrivate::GetDefaultSizeToAlignment(InSizeBytes) : InAlignment;
 			const uint32 AlignedNewTop = Align(Top, Alignment);
 		
 			if (AlignedNewTop + InSizeBytes < Page.Num())
