Fix audio unit tests

diff --git a/Engine/Plugins/Experimental/MetasoundExperimental/Source/MetasoundExperimentalRuntime/Private/MetasoundExperimentalRuntimeModule.cpp b/Engine/Plugins/Experimental/MetasoundExperimental/Source/MetasoundExperimentalRuntime/Private/MetasoundExperimentalRuntimeModule.cpp
--- a/Engine/Plugins/Experimental/MetasoundExperimental/Source/MetasoundExperimentalRuntime/Private/MetasoundExperimentalRuntimeModule.cpp
+++ b/Engine/Plugins/Experimental/MetasoundExperimental/Source/MetasoundExperimentalRuntime/Private/MetasoundExperimentalRuntimeModule.cpp
@@ -10,6 +10,7 @@
 
 void FMetasoundExperimentalRuntimeModule::StartupModule()
 {
+	FModuleManager::Get().LoadModule("AudioExperimentalRuntime");
 	FMetasoundFrontendRegistryContainer::Get()->RegisterPendingNodes();
 
 	// Datatypes.
