[Metasound][Cat] CatMixer node and fixes to transcoding around channel order

#tests AudioUnitTests

diff --git a/Engine/Plugins/Experimental/MetasoundExperimental/Source/AudioExperimentalRuntime/AudioExperimentalRuntime.Build.cs b/Engine/Plugins/Experimental/MetasoundExperimental/Source/AudioExperimentalRuntime/AudioExperimentalRuntime.Build.cs
--- a/Engine/Plugins/Experimental/MetasoundExperimental/Source/AudioExperimentalRuntime/AudioExperimentalRuntime.Build.cs
+++ b/Engine/Plugins/Experimental/MetasoundExperimental/Source/AudioExperimentalRuntime/AudioExperimentalRuntime.Build.cs
@@ -8,6 +8,9 @@ public class AudioExperimentalRuntime : ModuleRules
 	{
 		PCHUsage = PCHUsageMode.UseExplicitOrSharedPCHs;
 
+		// PLEASE DO NOT ADD AN ENGINE DEPENDENCY HERE!
+		// HORDE DOES NOT CATCH IT EITHER
+		// RUN (ushell): .build program AudioUnitTests OR build in VS
 		PublicDependencyModuleNames.AddRange(
 			new string[]
 			{
@@ -15,6 +18,14 @@ public class AudioExperimentalRuntime : ModuleRules
 			}
 		);
 		
+		// We want to use enum definitions in AudioMixer without an Engine dependency.
+		PublicIncludePathModuleNames.AddRange(
+			new string[]
+			{
+				"AudioMixer"
+			}
+		);
+		
 		// When the Editor builds all plugins, it will build with engine. In that case add a dependency on CoreUObject.
 		// If we don't do this we will get linker errors for the Engine expecting to find stuff defined here.
 		if (Target.bCompileAgainstEngine)
@@ -22,6 +33,9 @@ public class AudioExperimentalRuntime : ModuleRules
 			PublicDependencyModuleNames.Add("CoreUObject");
 		}
 		
+		// PLEASE DO NOT ADD AN ENGINE DEPENDENCY HERE!
+		// HORDE DOES NOT CATCH IT EITHER
+		// RUN (ushell): .build program AudioUnitTests OR build in VS
 		PrivateDependencyModuleNames.AddRange(
 			new string[]
 			{
diff --git a/Engine/Plugins/Experimental/MetasoundExperimental/Source/AudioExperimentalRuntime/Private/ChannelAgnostic/ChannelAgnosticType.cpp b/Engine/Plugins/Experimental/MetasoundExperimental/Source/AudioExperimentalRuntime/Private/ChannelAgnostic/ChannelAgnosticType.cpp
--- a/Engine/Plugins/Experimental/MetasoundExperimental/Source/AudioExperimentalRuntime/Private/ChannelAgnostic/ChannelAgnosticType.cpp
+++ b/Engine/Plugins/Experimental/MetasoundExperimental/Source/AudioExperimentalRuntime/Private/ChannelAgnostic/ChannelAgnosticType.cpp
@@ -47,4 +47,8 @@ namespace Audio
 		return Type->GetName();
 	}
 
+	void FChannelAgnosticType::Zero()
+	{
+		FMemory::Memzero(Buffer.GetView().GetData(), Buffer.GetView().NumBytes());
+	}
 }
diff --git a/Engine/Plugins/Experimental/MetasoundExperimental/Source/AudioExperimentalRuntime/Private/TypeFamily/ChannelTypeFamily.cpp b/Engine/Plugins/Experimental/MetasoundExperimental/Source/AudioExperimentalRuntime/Private/TypeFamily/ChannelTypeFamily.cpp
--- a/Engine/Plugins/Experimental/MetasoundExperimental/Source/AudioExperimentalRuntime/Private/TypeFamily/ChannelTypeFamily.cpp
+++ b/Engine/Plugins/Experimental/MetasoundExperimental/Source/AudioExperimentalRuntime/Private/TypeFamily/ChannelTypeFamily.cpp
@@ -140,6 +140,55 @@ namespace Audio
 		UE_LOG(LogTemp, Display, TEXT("Unique=%s\tNumChannels=%d\tParent=%s\tFriendlyName=%s\tDefault=%s\tOrder=[%s]\tAbstract=%s"),
 			*InUniqueName.ToString(),  InOrder.Num(), *GetParentNameSafe(InParentType), *InFriendlyName, ToCStr(LexToString(bIsParentsDefault)), *ChannelTypeFamilyPrviate::MakePrettyString(Order), ToCStr(LexToString(bIsAbstract)));
 	}
+
+	static TArrayView<ESpeakerShortNames> GetUnrealOrder()
+	{
+		static ESpeakerShortNames Order[]
+		{
+			#define CSV(X) X,
+			FOREACH_ENUM_ESPEAKERSHORTNAMES(CSV)
+			#undef CSV
+		};
+		return Order;
+	}
+
+	static void TranslateGainMatrix(
+		const TArrayView<float>& InGainMatrix,
+		const int32 InNumInputChannels,
+		const int32 InNumOutputChannels,
+		const TArrayView<const ESpeakerShortNames>& InOldOrder,
+		const TArrayView<const ESpeakerShortNames>& InNewOrder,
+		TArray<float>& OutGains)
+	{
+		// Channels not found in new order will be zeroed.
+		OutGains.SetNumZeroed(InNumOutputChannels*InNumInputChannels);
+		
+		// Lookup (old->new).
+		int32 Lookup[static_cast<int32>(ESpeakerShortNames::NumChannels)] = { INDEX_NONE };
+		for (int32 i = 0; i < InOldOrder.Num(); ++i)
+		{
+			Lookup[static_cast<int32>(InOldOrder[i])] = InNewOrder.IndexOfByKey(InOldOrder[i]);
+		}
+		
+		// Translate.
+		for (int32 InputChannelIndex=0; InputChannelIndex < InNumInputChannels; ++InputChannelIndex)
+		{
+			// This input channel has a mapping in the new?
+			const ESpeakerShortNames OldIn = InOldOrder[InputChannelIndex];
+			if (const int32 NewIn = Lookup[static_cast<int32>(OldIn)]; NewIn != INDEX_NONE)
+			{
+				for (int32 OutputChannelIndex=0; OutputChannelIndex < InNumOutputChannels; ++OutputChannelIndex)
+				{
+					const ESpeakerShortNames OldOut = InOldOrder[OutputChannelIndex];
+					if (const int32 NewOut = Lookup[static_cast<int32>(OldOut)]; NewOut != INDEX_NONE)
+					{
+							OutGains[(NewIn				* InNumOutputChannels)		+ NewOut] =
+						InGainMatrix[(InputChannelIndex * InNumOutputChannels)		+ OutputChannelIndex];
+					}
+				}
+			}
+		}
+	}
 	
 	FChannelTypeFamily::FTranscoder FDiscreteChannelTypeFamily::GetTranscoder(
 	const FDiscreteChannelTypeFamily& InFromType, const FGetTranscoderParams& InParams) const
@@ -189,7 +238,9 @@ namespace Audio
 						.bIsCenterChannelOnly = InFromType.NumChannels() == 1 && InFromType.HasSpeaker(ESpeakerShortNames::FC)
 					}, Gains))
 				{
-					return [MixGains = MoveTemp(Gains)](TArrayView<const float*> InSrc, TArrayView<float*> InDst, const int32 NumFrames)
+					TArray<float> TranslatedGains;
+					TranslateGainMatrix(Gains, InFromType.NumChannels(), NumChannels(), GetUnrealOrder(), InFromType.GetSpeakerOrder(), TranslatedGains);
+					return [MixGains = MoveTemp(TranslatedGains)](TArrayView<const float*> InSrc, TArrayView<float*> InDst, const int32 NumFrames)
 					{
 						MultiMonoMixUpOrDown(InSrc, InDst, NumFrames, MixGains);
 					};
@@ -241,8 +292,8 @@ Registry.RegisterType(TEXT(UE_STRINGIZE(NAME)),UE_JOIN(__Type,NAME))
 	
 		// Mono		(UniqueName)		(Parent Name)			(Friendly Name)										(Channel membership/order.)				(is parents default).  (abstract)
 		REGISTER_CAT(Mono,				TEXT("Discrete"),TEXT("Mono"),							{},										false,					true);
-		REGISTER_CAT(Mono1Dot0,			TEXT("Mono"),	TEXT("Mono (1.0)"),					(TArray{FC}),							true,					false);
-		REGISTER_CAT(Mono1Dot1,			TEXT("Mono"),	TEXT("Mono (1.1)"),					(TArray{FC,LFE}),						false,					false);
+		REGISTER_CAT(Mono1Dot0,			TEXT("Mono"),	TEXT("Mono (1.0)"),					(TArray{FL}),							true,					false);
+		REGISTER_CAT(Mono1Dot1,			TEXT("Mono"),	TEXT("Mono (1.1)"),					(TArray{FL,LFE}),						false,					false);
 
 		// Stereo																								
 		REGISTER_CAT(Stereo,			TEXT("Discrete"), TEXT("Stereo"),						{},										false,					true);
diff --git a/Engine/Plugins/Experimental/MetasoundExperimental/Source/AudioExperimentalRuntime/Public/ChannelAgnostic/ChannelAgnosticType.h b/Engine/Plugins/Experimental/MetasoundExperimental/Source/AudioExperimentalRuntime/Public/ChannelAgnostic/ChannelAgnosticType.h
--- a/Engine/Plugins/Experimental/MetasoundExperimental/Source/AudioExperimentalRuntime/Public/ChannelAgnostic/ChannelAgnosticType.h
+++ b/Engine/Plugins/Experimental/MetasoundExperimental/Source/AudioExperimentalRuntime/Public/ChannelAgnostic/ChannelAgnosticType.h
@@ -61,6 +61,11 @@ namespace Audio
 		FName GetTypeName() const;
 		const FChannelTypeFamily& GetType() const { return *Type; }
 
+		/*
+		 * Zero the contents of the buffer.
+		 */
+		void Zero();
+
 		/**
 		 * For fast raw DSP access to the buffer.
 		 * @return View of entire multi-mono buffer.
diff --git a/Engine/Plugins/Experimental/MetasoundExperimental/Source/AudioExperimentalRuntime/Public/TypeFamily/ChannelTypeFamily.h b/Engine/Plugins/Experimental/MetasoundExperimental/Source/AudioExperimentalRuntime/Public/TypeFamily/ChannelTypeFamily.h
--- a/Engine/Plugins/Experimental/MetasoundExperimental/Source/AudioExperimentalRuntime/Public/TypeFamily/ChannelTypeFamily.h
+++ b/Engine/Plugins/Experimental/MetasoundExperimental/Source/AudioExperimentalRuntime/Public/TypeFamily/ChannelTypeFamily.h
@@ -2,31 +2,35 @@
 
 #pragma once
 
+#include "AudioMixer.h"
 #include "TypeFamily.h"
 #include "DSP/ChannelMap.h"
 #include "Templates/Function.h"
 #include "Misc/Optional.h"
 
+
 /**
  * A Comprehensive Short form Speaker enumeration
  */
 enum class ESpeakerShortNames : uint32
 {
-	FL,   // Front Left
-	FR,   // Front Right
-	FC,   // Front Center
-	LFE,  // Low-Frequency Effects (Subwoofer)
-	FLC,  // Front Left Center
-	FRC,  // Front Right Center
-	SL,   // Side Left
-	SR,   // Side Right
-	BL,   // Back Left
-	BR,   // Back Right
-	BC,   // Back Center
-	TFL,  // Top Front Left
-	TFR,  // Top Front Right
-	TBL,  // Top Back Left
-	TBR   // Top Back Right
+	FL = EAudioMixerChannel::FrontLeft,				// Front Left
+	FR = EAudioMixerChannel::FrontRight,			// Front Right
+	FC = EAudioMixerChannel::FrontCenter,			// Front Center
+	LFE = EAudioMixerChannel::LowFrequency,			// Low-Frequency Effects (Subwoofer)
+	FLC = EAudioMixerChannel::FrontLeftOfCenter,	// Front Left Center
+	FRC = EAudioMixerChannel::FrontRightOfCenter,	// Front Right Center
+	SL = EAudioMixerChannel::SideLeft,				// Side Left
+	SR = EAudioMixerChannel::SideRight,				// Side Right
+	BL = EAudioMixerChannel::BackLeft,				// Back Left
+	BR = EAudioMixerChannel::BackRight,				// Back Right
+	BC = EAudioMixerChannel::BackCenter,			// Back Center
+	TFL = EAudioMixerChannel::TopFrontLeft,			// Top Front Left
+	TFR = EAudioMixerChannel::TopFrontRight,		// Top Front Right
+	TBL = EAudioMixerChannel::TopBackLeft,			// Top Back Left
+	TBR = EAudioMixerChannel::TopBackRight,			// Top Back Right
+
+	NumChannels = EAudioMixerChannel::ChannelTypeCount,
 };
 
 // Pasted from UHT generated.h, but we can't use because of AudioUnitTests.
@@ -162,6 +166,19 @@ namespace Audio
 		};
 		virtual TOptional<FChannelName> GetChannelName(const int32 InChannelIndex) const { return {}; }
 
+		template<typename T>
+		T* Cast()
+		{
+			if (T::GetFamilyTypeName() == GetFamilyName())
+			{
+				return static_cast<T*>(this);
+			}
+			return nullptr;
+		}
+
+		template<typename T>
+		const T* Cast() const { return const_cast<FChannelTypeFamily*>(this)->Cast<T>(); }
+
 	protected:
 		// Both are friends so we can hide the api.
 		friend class FDiscreteChannelTypeFamily;
@@ -223,6 +240,12 @@ namespace Audio
 
 		virtual TOptional<FChannelName> GetChannelName(const int32 InChannelIndex) const override;
 
+		/**
+		 * Get the order of the speakers defined on this discrete format.
+		 * @return Array of speakers in order.
+		 */
+		const TArray<ESpeakerShortNames>& GetSpeakerOrder() const { return Order; }
+
 	protected:
 		// Visit "to" Type.
 		virtual FTranscoder GetTranscoder(const FGetTranscoderParams& InParams) const override
diff --git a/Engine/Plugins/Experimental/MetasoundExperimental/Source/MetasoundExperimentalRuntime/MetasoundExperimentalRuntime.Build.cs b/Engine/Plugins/Experimental/MetasoundExperimental/Source/MetasoundExperimentalRuntime/MetasoundExperimentalRuntime.Build.cs
--- a/Engine/Plugins/Experimental/MetasoundExperimental/Source/MetasoundExperimentalRuntime/MetasoundExperimentalRuntime.Build.cs
+++ b/Engine/Plugins/Experimental/MetasoundExperimental/Source/MetasoundExperimentalRuntime/MetasoundExperimentalRuntime.Build.cs
@@ -8,6 +8,10 @@ public class MetasoundExperimentalRuntime : ModuleRules
 	{
 		PCHUsage = ModuleRules.PCHUsageMode.UseExplicitOrSharedPCHs;
 
+		// PLEASE DO NOT ADD AN ENGINE DEPENDENCY HERE!
+		// HORDE DOES NOT CATCH IT EITHER
+		// RUN: .build program AudioUnitTests
+		// RUN (ushell): .build program AudioUnitTests OR build in VS
 		PublicDependencyModuleNames.AddRange(
 			new string[]
 			{
@@ -25,13 +29,17 @@ public class MetasoundExperimentalRuntime : ModuleRules
 			PublicDependencyModuleNames.Add("CoreUObject");
 		}
 		
+		// PLEASE DO NOT ADD AN ENGINE DEPENDENCY HERE!
+		// HORDE DOES NOT CATCH IT EITHER
+		// RUN (ushell): .build program AudioUnitTests OR build in VS
 		PrivateDependencyModuleNames.AddRange(
 			new string[]
 			{
 				"Core",
 				"MetasoundGraphCore",
 				"MetasoundFrontend",
-				"MetasoundStandardNodes"
+				"MetasoundStandardNodes",
+				"SignalProcessing"
 			}
 		);
 	}
diff --git a/Engine/Plugins/Experimental/MetasoundExperimental/Source/MetasoundExperimentalRuntime/Private/MetasoundCatMixerNode.cpp b/Engine/Plugins/Experimental/MetasoundExperimental/Source/MetasoundExperimentalRuntime/Private/MetasoundCatMixerNode.cpp
--- a/Engine/Plugins/Experimental/MetasoundExperimental/Source/MetasoundExperimentalRuntime/Private/MetasoundCatMixerNode.cpp
+++ b/Engine/Plugins/Experimental/MetasoundExperimental/Source/MetasoundExperimentalRuntime/Private/MetasoundCatMixerNode.cpp
@@ -9,7 +9,11 @@
 #include "MetasoundPrimitives.h"
 #include "MetasoundStandardNodesCategories.h"
 #include "MetasoundVertex.h"
+#include "Algo/MaxElement.h"
+#include "DSP/FloatArrayMath.h"
 #include "TypeFamily/ChannelTypeFamily.h"
+#include "MetasoundCatBreakNode.h"
+#include "ChannelAgnostic/ChannelAgnosticTypeUtils.h"
 
 #define LOCTEXT_NAMESPACE "MetasoundStandardNodes_CatMixerNode"
 
@@ -17,57 +21,159 @@ namespace Metasound
 {
 	namespace CatMixerPrivate
 	{
-		METASOUND_PARAM(InputGain,			"Gain", "Channel Agnostic Output");
-		METASOUND_PARAM(InputCat,			"Cat In", "Channel Agnostic Input");
-		METASOUND_PARAM(InputNumChannels,	"NumChannels", "Num Output Channels");
 		METASOUND_PARAM(OutputCat,			"Cat Out", "Channel Agnostic Output");
-		const FString DefaultCatFormat = TEXT("Mono_1_0");
+		
+		class FCatMixerOperatorData final : public TOperatorData<FCatMixerOperatorData>
+		{
+		public:
+			static const FLazyName OperatorDataTypeName;
+			FCatMixerOperatorData(const FName InToMixType, const int32 InNumInputs, const EMetasoundMixerFormatChoosingMethod InFormatChosingMethod, const EMetasoundChannelMapMonoUpmixMethod InUpMixMethod, const EMetasoundCatCastingMethod InCastingMethod )
+				: ToMixType(InToMixType), NumInputs(InNumInputs), FormatChoosingMethod(InFormatChosingMethod), ChannelMapMonoUpmixMethod(InUpMixMethod), CatCastingMethod(InCastingMethod)
+			{}
+			FName ToMixType;
+			int32 NumInputs = 1;
+			EMetasoundMixerFormatChoosingMethod FormatChoosingMethod;
+			EMetasoundChannelMapMonoUpmixMethod ChannelMapMonoUpmixMethod;
+			EMetasoundCatCastingMethod CatCastingMethod;
+		};
+
+		// Linkage.
+		const FLazyName FCatMixerOperatorData::OperatorDataTypeName = TEXT("FCatMixerOperatorData");
+
+		const FLazyName CatInputBaseName{"In"};
+		const FLazyName GainInputBaseName{ "Gain" };
+
+		static FName MakeCatInputVertexName(const int32 InIndex)
+		{
+			FName Name = CatInputBaseName;
+			Name.SetNumber(InIndex);
+			return Name;
+		}
+
+		static FName MakeGainInputVertexName(const int32 InIndex)
+		{
+			FName Name = GainInputBaseName;
+			Name.SetNumber(InIndex);
+			return Name;
+		}
+
+		static FInputDataVertex MakeInputCatVertex(const int32 InIndex, const FName InName, const FString& InFriendlyName)
+		{
+			const FName InputName = MakeCatInputVertexName(InIndex);
+#if WITH_EDITOR
+			const FText InputDisplayName = FText::Format(LOCTEXT("In_DisplayName", "{0}"), FText::FromString(FString::Printf(TEXT("%s %d"), *InName.ToString(), InIndex)));
+#else
+			const FText InputDisplayName = FText::GetEmpty();
+#endif
+			return TInputDataVertex<FChannelAgnosticType>{InputName, FDataVertexMetadata{FText::FromString(*InFriendlyName), InputDisplayName}};
+		}
+
+		static FInputDataVertex MakeInputGainVertex(const int32 InIndex, const FName InName, const FString& InFriendlyName)
+		{
+			const FName InputName = MakeGainInputVertexName(InIndex);
+#if WITH_EDITOR
+			const FText InputDisplayName = FText::Format(LOCTEXT("In_DisplayName", "{0}"), FText::FromString(FString::Printf(TEXT("%s %d"), *InName.ToString(), InIndex)));
+#else
+			const FText InputDisplayName = FText::GetEmpty();
+#endif
+			return TInputDataVertex<float>{InputName, FDataVertexMetadata{ FText::FromString(*InFriendlyName), InputDisplayName }, 1.0f};
+		}
+
+		FVertexInterface GetVertexInterface(const int32 InNumInputs)
+		{
+			FInputVertexInterface InputInterface;
+			for (int32 i = 0; i < InNumInputs; i++)
+			{
+				InputInterface.Add(MakeInputCatVertex(i,CatInputBaseName, TEXT("Cat Input") ));
+				InputInterface.Add(MakeInputGainVertex(i, GainInputBaseName, TEXT("Gain Input")));
+			}
+
+			FOutputVertexInterface OutputInterface;
+			OutputInterface.Add(TOutputDataVertex<FChannelAgnosticType>(METASOUND_GET_PARAM_NAME_AND_METADATA(OutputCat)));
+			return FVertexInterface
+				{
+					MoveTemp(InputInterface),
+					MoveTemp(OutputInterface)
+				};
+		}
+
+		static const Audio::FChannelTypeFamily* FindHighestInputType(const TArray<FChannelAgnosticTypeReadRef>& Inputs)
+		{
+			if (const FChannelAgnosticTypeReadRef* Found = Algo::MaxElementBy(Inputs, [](const FChannelAgnosticTypeReadRef& i) { return i->NumChannels(); }) )
+			{
+				return &(*Found)->GetType();
+			}
+			return nullptr;
+		}
+		static const Audio::FChannelTypeFamily* FindLowestInputType(const TArray<FChannelAgnosticTypeReadRef>& Inputs)
+		{
+			if (const FChannelAgnosticTypeReadRef* Found = Algo::MinElementBy(Inputs, [](const FChannelAgnosticTypeReadRef& i) { return i->NumChannels(); }) )
+			{
+				return &(*Found)->GetType();
+			}
+			return nullptr;
+		}
+		static const Audio::FChannelTypeFamily* ChooseOutputFormat(
+			const TArray<FChannelAgnosticTypeReadRef>& InOutputs, const Audio::IChannelTypeRegistry& InRegistry, const FCatMixerOperatorData& InOperatorData)
+		{
+			using enum EMetasoundMixerFormatChoosingMethod;
+			switch (InOperatorData.FormatChoosingMethod)
+			{
+			case MetasoundOutput:	// TODO. 
+			case HighestInput:		return FindHighestInputType(InOutputs);
+			case LowestInput:		return FindLowestInputType(InOutputs);
+			case Custom:			return InRegistry.FindChannel(InOperatorData.ToMixType);
+			default:				checkNoEntry();
+			}
+			return nullptr;
+		}
+		static const Audio::FChannelTypeFamily* ResolveOutputFormat(
+		const TArray<FChannelAgnosticTypeReadRef>& InOutputs, const Audio::IChannelTypeRegistry& InRegistry, const FCatMixerOperatorData& InOperatorData)
+		{
+			if (const Audio::FChannelTypeFamily* Found = ChooseOutputFormat(InOutputs, InRegistry, InOperatorData))
+			{
+				return InRegistry.FindConcreteChannel(Found->GetName());
+			}
+			return nullptr;
+		}
 	}
 	
 	class FCatMixerOperator final : public TExecutableOperator<FCatMixerOperator>
 	{
 	public:
-		FCatMixerOperator(const FBuildOperatorParams& InParams, FChannelAgnosticTypeReadRef&& InInputCat, FFloatReadRef&& InGain, FInt32ReadRef&& InNumOutputChannels)
-			: Gain(MoveTemp(InGain))
-			, Inputs (MoveTemp(InInputCat))
-			, NumOutputChannels(MoveTemp(InNumOutputChannels))
-			, Outputs(FChannelAgnosticTypeWriteRef::CreateNew(InParams.OperatorSettings, CatMixerPrivate::DefaultCatFormat)) // initalize to a default format a head of binding. (maybe default init is better?)
-			, Settings(InParams.OperatorSettings)
+		FCatMixerOperator(TArray<FFloatReadRef>&& InGains, TArray<FChannelAgnosticTypeReadRef>&& InInputCat, FChannelAgnosticTypeWriteRef&& InOutputCat,
+			const FOperatorSettings& InSettings, const CatMixerPrivate::FCatMixerOperatorData& InOpData)
+			: Gains(MoveTemp(InGains))
+			, Inputs(MoveTemp(InInputCat))
+			, Outputs(MoveTemp(InOutputCat))
+			, Settings(InSettings)
+			, TranscodeScratch(InSettings, Outputs->GetType().GetName())
 		{
-			Reset(InParams);
-		}
-		virtual ~FCatMixerOperator() override = default;
-
-		static const FVertexInterface& GetDefaultInterface()
-		{
-			using namespace CatMixerPrivate;
-			auto CreateDefaultInterface = []()-> FVertexInterface
+			PrevGains.Reset(Gains.Num());
+			for (FFloatReadRef& Gain : Gains)
 			{
-				// inputs
-				FInputVertexInterface InputInterface;
-				InputInterface.Add(TInputDataVertex<FChannelAgnosticType>(METASOUND_GET_PARAM_NAME_AND_METADATA(InputCat),CatMixerPrivate::DefaultCatFormat));
-				InputInterface.Add(TInputDataVertex<float>(METASOUND_GET_PARAM_NAME_AND_METADATA(InputGain)));
-				InputInterface.Add(TInputDataVertex<int32>(METASOUND_GET_PARAM_NAME_AND_METADATA(InputNumChannels)));
-				
-				// outputs
-				FOutputVertexInterface OutputInterface;
-				OutputInterface.Add(TOutputDataVertex<FChannelAgnosticType>(METASOUND_GET_PARAM_NAME_AND_METADATA(OutputCat)));
-				
-				return FVertexInterface(InputInterface, OutputInterface);
-			}; // end lambda: CreateDefaultInterface()
-
-			static const FVertexInterface DefaultInterface = CreateDefaultInterface();
-			return DefaultInterface;
+				PrevGains.Add(*Gain);
+			}
+			for  (FChannelAgnosticTypeReadRef& Input : Inputs)
+			{
+				Transcoders.Emplace(Input->GetType().GetTranscoder(
+					{
+						.ToType = Outputs->GetType(),
+						.TranscodeMethod = static_cast<Audio::EChannelTranscodeMethod>(InOpData.CatCastingMethod),
+						.MixMethod = static_cast<Audio::EChannelMapMonoUpmixMethod>(InOpData.CatCastingMethod)
+					}));
+			}
 		}
+		virtual ~FCatMixerOperator() override = default;
 
 		static const FNodeClassMetadata& GetNodeInfo()
 		{
 			static const FNodeClassMetadata Metadata = CreateNodeClassMetadata
 			(
-				FName(TEXT("Cat Mixer Node")),
-				METASOUND_LOCTEXT("Metasound_CatMixerNodeDisplayName", "Cat Mixer Node"),
-				METASOUND_LOCTEXT("Metasound_CatMixerNodeDescription", "Cat Mixer Node"),
-				GetDefaultInterface()	
+				FName(TEXT("CAT Mixer Node")),
+				METASOUND_LOCTEXT("Metasound_CatMixerNodeDisplayName", "CAT Mixer Node"),
+				METASOUND_LOCTEXT("Metasound_CatMixerNodeDescription", "CAT Mixer Node"),
+				CatMixerPrivate::GetVertexInterface(1)
 			);
 			return Metadata;
 		}
@@ -76,61 +182,119 @@ namespace Metasound
 		{
 			using namespace CatMixerPrivate;
 
-			const FInputVertexInterfaceData& InputData = InParams.InputData;
+			const FCatMixerOperatorData* OperatorData = CastOperatorData<const FCatMixerOperatorData>(InParams.Node.GetOperatorData().Get());
+			if (!OperatorData)
+			{
+				return MakeUnique<FNoOpOperator>();
+			}
+				
+			const Audio::FChannelTypeFamily* ConcreteToType = Audio::GetChannelRegistry().FindConcreteChannel(OperatorData->ToMixType);
+			if (!ConcreteToType)
+			{
+				return MakeUnique<FNoOpOperator>();
+			}
+
+			// inputs.
+			TArray<FChannelAgnosticTypeReadRef> Inputs;
+			TArray<FFloatReadRef> Gains;
+			for (int32 i = 0; i < OperatorData->NumInputs; ++i)
+			{
+				Inputs.Emplace(InParams.InputData.GetOrCreateDefaultDataReadReference<FChannelAgnosticType>(MakeCatInputVertexName(i), InParams.OperatorSettings));
+				Gains.Emplace(InParams.InputData.GetOrCreateDefaultDataReadReference<float>(MakeGainInputVertexName(i), InParams.OperatorSettings));
+			}
+
+			const Audio::FChannelTypeFamily* Resolved = ResolveOutputFormat(Inputs,Audio::GetChannelRegistry(), *OperatorData);
+			if (!Resolved)
+			{
+				return MakeUnique<FNoOpOperator>();
+			}
+			
+			// Make the node.
 			return MakeUnique<FCatMixerOperator>(
-				InParams,
-				InputData.GetOrCreateDefaultDataReadReference<FChannelAgnosticType>(METASOUND_GET_PARAM_NAME(InputCat), InParams.OperatorSettings),
-				InputData.GetOrCreateDefaultDataReadReference<float>(METASOUND_GET_PARAM_NAME(InputGain), InParams.OperatorSettings),
-				InputData.GetOrCreateDefaultDataReadReference<int32>(METASOUND_GET_PARAM_NAME(InputNumChannels), InParams.OperatorSettings)
+				MoveTemp(Gains),
+				MoveTemp(Inputs),
+				FChannelAgnosticTypeWriteRef::CreateNew(InParams.OperatorSettings, Resolved->GetName()),
+				InParams.OperatorSettings,
+				*OperatorData
 			);
 		}
 
 		virtual void BindInputs(FInputVertexInterfaceData& InOutVertexData) override
 		{
 			using namespace CatMixerPrivate;
-			InOutVertexData.BindReadVertex(METASOUND_GET_PARAM_NAME(InputGain) , Gain);
-			InOutVertexData.BindReadVertex(METASOUND_GET_PARAM_NAME(InputCat) , Inputs);
-			InOutVertexData.BindReadVertex(METASOUND_GET_PARAM_NAME(InputNumChannels) , NumOutputChannels);
+			for (int32 i = 0; i < Inputs.Num(); i++)
+			{
+				InOutVertexData.BindReadVertex(MakeCatInputVertexName(i), Inputs[i]);
+				InOutVertexData.BindReadVertex(MakeGainInputVertexName(i), Gains[i]);
+			}
 		}
 
 		virtual void BindOutputs(FOutputVertexInterfaceData& InOutVertexData) override
 		{
 			using namespace CatMixerPrivate;
+			InOutVertexData.BindWriteVertex(METASOUND_GET_PARAM_NAME(OutputCat), Outputs);
+		}
 
-			if (const int32 NumChannels = *NumOutputChannels; NumChannels > 0) 
+		void Reset(const FResetParams&)
+		{
+			Outputs->Zero();
+			for (int32 i = 0; i < Inputs.Num(); i++)
 			{
-				// Create cat type capable of holding N channels.
-				// FIXME: Use FName chooser.
-				if (const Audio::FChannelTypeFamily* Type = Audio::GetChannelRegistry().FindChannelType(NumChannels))
-				{
-					Outputs = FChannelAgnosticTypeWriteRef::CreateNew(Settings, Type->GetName().ToString());
-				}
-			
-				// Bind it
-				InOutVertexData.BindWriteVertex(METASOUND_GET_PARAM_NAME(OutputCat), Outputs);
+				PrevGains[i] = 0.f;
 			}
 		}
 
-		void Reset(const IOperator::FResetParams& InParams)
+		static void MixInDiscreteCat(const FChannelAgnosticType& InSrc, FChannelAgnosticType& InDst, const float InStartGain, const float InEndGain)
 		{
-			PrevGains = *Gain;
-			Execute();
+			// these are the same format now as they've been transcoded.
+			check(InSrc.NumChannels() == InDst.NumChannels());
+			check(InSrc.IsA(InDst.GetType().GetName()));
+			
+			const int32 NumChannels = InSrc.NumChannels();
+			for (int32 i = 0; i < NumChannels; ++i)
+			{
+				TArrayView<float> Dst = InDst.GetChannel(i);
+				TArrayView<const float> Src = InSrc.GetChannel(i);
+				Audio::ArrayMixIn(Src, Dst, InStartGain, InEndGain);
+			}
 		}
-
+		
 		void Execute()
 		{
+			using namespace Audio;
+
+			// Always zero the output to start.
+			Outputs->Zero();
+
+			const int32 NumFrames = Settings.GetNumFramesPerBlock();
+			for (int32 i = 0; i < Inputs.Num(); i++)
+			{
+				const FChannelAgnosticTypeReadRef& Input = Inputs[i]; 
+				const float PrevGainValue = PrevGains[i];
+				const float GainValue = *Gains[i];
+				{
+					// Transcode channel into scratch buffer.
+					TStackArrayOfPointers<const float> Src = MakeMultiMonoPointersFromView(Input->GetRawMultiMono(), NumFrames, Input->NumChannels());
+					TStackArrayOfPointers<float> Dst = MakeMultiMonoPointersFromView(TranscodeScratch.GetRawMultiMono(), NumFrames, Outputs->NumChannels());
+					Transcoders[i](Src, Dst, NumFrames);
+				}
+
+				// Mix scratch buffer into results. (lerp previous gain and current to avoid zippering).
+				MixInDiscreteCat(TranscodeScratch, *Outputs , PrevGainValue, GainValue);
+
+				PrevGains[i] = GainValue;
+			}
 		}
 		
 	private:
 
-		FFloatReadRef Gain;
-		FChannelAgnosticTypeReadRef Inputs;
-		FInt32ReadRef NumOutputChannels;
-		int32 NumFrames = 0;
-		int32 NumInputChannels = 0;
+		TArray<FFloatReadRef> Gains;
+		TArray<float> PrevGains;
+		TArray<FChannelAgnosticTypeReadRef> Inputs;
 		FChannelAgnosticTypeWriteRef Outputs;
 		FOperatorSettings Settings;
-		float PrevGains = 0.f;
+		TArray<Audio::FChannelTypeFamily::FTranscoder> Transcoders;
+		FChannelAgnosticType TranscodeScratch;
 
 		static FNodeClassMetadata CreateNodeClassMetadata(const FName& InOperatorName, const FText& InDisplayName, const FText& InDescription, const FVertexInterface& InDefaultInterface)
 		{
@@ -155,6 +319,23 @@ namespace Metasound
 
 	using FCatMixerNode = TNodeFacade<FCatMixerOperator>;
 
+	// register node config.	
+	METASOUND_REGISTER_NODE_AND_CONFIGURATION(FCatMixerNode, FMetaSoundCatMixingNodeConfiguration);
+	
 } // namespace Metasound
 
+TInstancedStruct<FMetasoundFrontendClassInterface> FMetaSoundCatMixingNodeConfiguration::OverrideDefaultInterface(const FMetasoundFrontendClass& InNodeClass) const
+{
+	using namespace Metasound;
+	using namespace CatMixerPrivate;
+	return TInstancedStruct<FMetasoundFrontendClassInterface>::Make(FMetasoundFrontendClassInterface::GenerateClassInterface(GetVertexInterface(NumInputs)));
+}
+
+TSharedPtr<const Metasound::IOperatorData> FMetaSoundCatMixingNodeConfiguration::GetOperatorData() const
+{
+	using namespace Metasound::CatMixerPrivate;
+	return MakeShared<FCatMixerOperatorData>(CustomMixFormat, NumInputs, FormatChoosingMethod, ChannelMapMonoUpmixMethod, CatCastingMethod);
+}
+
+
 #undef LOCTEXT_NAMESPACE // "MetasoundStandardNodes_CatMixerNode"
diff --git a/Engine/Plugins/Experimental/MetasoundExperimental/Source/MetasoundExperimentalRuntime/Private/MetasoundCatMixerNode.h b/Engine/Plugins/Experimental/MetasoundExperimental/Source/MetasoundExperimentalRuntime/Private/MetasoundCatMixerNode.h
--- a/Engine/Plugins/Experimental/MetasoundExperimental/Source/MetasoundExperimentalRuntime/Private/MetasoundCatMixerNode.h
+++ b/Engine/Plugins/Experimental/MetasoundExperimental/Source/MetasoundExperimentalRuntime/Private/MetasoundCatMixerNode.h
@@ -1,23 +1,45 @@
 ﻿// Copyright Epic Games, Inc. All Rights Reserved.
 
-#if 0
 
 #pragma once
 
-#include "MetasoundFacade.h"
-#include "MetasoundNodeInterface.h"
+#include "MetasoundCatCastingNode.h"
+#include "MetasoundFrontendDocument.h"
 #include "Templates/SharedPointer.h"
 
-namespace Metasound
+#include "MetasoundCatMixerNode.generated.h"
+
+UENUM()
+enum class EMetasoundMixerFormatChoosingMethod: uint8
+{
+	HighestInput UMETA(DisplayName = "Highest Num Channels", ToolTip = "Output will be the MAX Channel count of the input channel"),
+	LowestInput UMETA(DisplayName = "Lowest Num Channels", ToolTip = "Output will be the MIN Channel count of the input channel"),
+	MetasoundOutput UMETA(DisplayName = "Metasound Output", ToolTip = "Output will match the Metasound Output setting"),
+	Custom UMETA(DisplayName = "Custom", ToolTip = "Custom mix format"),
+};
+
+USTRUCT()
+struct FMetaSoundCatMixingNodeConfiguration : public FMetaSoundFrontendNodeConfiguration
 {
-	class FCatMixerNode final : public FNodeFacade
-	{
-	public:
-		explicit FCatMixerNode(FNodeData InNodeData, TSharedRef<const FNodeClassMetadata> InClassMetadata);
-		virtual ~FCatMixerNode() override = default;
+	GENERATED_BODY()
+
+	FMetaSoundCatMixingNodeConfiguration() = default;
+
+	UPROPERTY(EditAnywhere, Category = General, DisplayName = "Mix Format Method")
+	EMetasoundMixerFormatChoosingMethod FormatChoosingMethod = EMetasoundMixerFormatChoosingMethod::HighestInput;
+
+	UPROPERTY(EditAnywhere, Category = General, meta=(EditCondition = "CatCastingMethod == EMetasoundCatCastingMethod::MixUpOrDown", EditConditionHides))
+	EMetasoundChannelMapMonoUpmixMethod ChannelMapMonoUpmixMethod = EMetasoundChannelMapMonoUpmixMethod::EqualPower;
+
+	UPROPERTY(EditAnywhere, Category = General)
+	EMetasoundCatCastingMethod CatCastingMethod = EMetasoundCatCastingMethod::ChannelDrop;
 
-		static FNodeClassMetadata CreateNodeClassMetadata();
-	};
-}
+	UPROPERTY(EditAnywhere, Category = General, meta = (GetOptions="MetasoundCatCastingOptionsHelper.GetCastingOptions", EditCondition = "FormatChoosingMethod == EMetasoundMixerFormatChoosingMethod::Custom", EditConditionHides))
+	FName CustomMixFormat = TEXT("Mono");
 
-#endif //
\ No newline at end of file
+	UPROPERTY(EditAnywhere, Category = General, meta = (ClampMin=1, ClampMax=100))
+	int32 NumInputs = 1;
+	
+	virtual TInstancedStruct<FMetasoundFrontendClassInterface> OverrideDefaultInterface(const FMetasoundFrontendClass& InNodeClass) const override;
+	virtual TSharedPtr<const Metasound::IOperatorData> GetOperatorData() const override;
+};
