CAT Ambisonics

diff --git a/Engine/Plugins/Experimental/MetasoundExperimental/Source/AudioExperimentalRuntime/Private/ChannelAgnostic/ChannelAgnosticTranscoding.cpp b/Engine/Plugins/Experimental/MetasoundExperimental/Source/AudioExperimentalRuntime/Private/ChannelAgnostic/ChannelAgnosticTranscoding.cpp
new file mode 100644
--- /dev/null
+++ b/Engine/Plugins/Experimental/MetasoundExperimental/Source/AudioExperimentalRuntime/Private/ChannelAgnostic/ChannelAgnosticTranscoding.cpp
@@ -0,0 +1,238 @@
+﻿// Copyright Epic Games, Inc. All Rights Reserved.
+
+#include "ChannelAgnostic/ChannelAgnosticTranscoding.h"
+
+#include "Algo/Transform.h"
+#include "DSP/Ambisonics.h"
+#include "DSP/MultiMono.h"
+#include "DSP/SphericalHarmonicCalculator.h"
+#include "TypeFamily/ChannelTypeFamily.h"
+
+namespace Audio
+{
+	static TArrayView<ESpeakerShortNames> GetUnrealOrder()
+	{
+		static ESpeakerShortNames Order[]
+		{
+			#define CSV(X) X,
+			FOREACH_ENUM_ESPEAKERSHORTNAMES(CSV)
+			#undef CSV
+		};
+		return Order;
+	}
+
+	static void TranslateGainMatrix(
+		const TArrayView<float>& InGainMatrix,
+		const int32 InNumInputChannels,
+		const int32 InNumOutputChannels,
+		const TArrayView<const ESpeakerShortNames>& InOldOrder,
+		const TArrayView<const ESpeakerShortNames>& InNewOrder,
+		TArray<float>& OutGains)
+	{
+		// Channels not found in new order will be zeroed.
+		OutGains.SetNumZeroed(InNumOutputChannels*InNumInputChannels);
+		
+		// Lookup (old->new).
+		int32 Lookup[static_cast<int32>(ESpeakerShortNames::NumChannels)] = { INDEX_NONE };
+		for (int32 i = 0; i < InOldOrder.Num(); ++i)
+		{
+			Lookup[static_cast<int32>(InOldOrder[i])] = InNewOrder.IndexOfByKey(InOldOrder[i]);
+		}
+		
+		// Translate.
+		for (int32 InputChannelIndex=0; InputChannelIndex < InNumInputChannels; ++InputChannelIndex)
+		{
+			// This input channel has a mapping in the new?
+			const ESpeakerShortNames OldIn = InOldOrder[InputChannelIndex];
+			if (const int32 NewIn = Lookup[static_cast<int32>(OldIn)]; NewIn != INDEX_NONE)
+			{
+				for (int32 OutputChannelIndex=0; OutputChannelIndex < InNumOutputChannels; ++OutputChannelIndex)
+				{
+					const ESpeakerShortNames OldOut = InOldOrder[OutputChannelIndex];
+					if (const int32 NewOut = Lookup[static_cast<int32>(OldOut)]; NewOut != INDEX_NONE)
+					{
+						OutGains[(NewIn				* InNumOutputChannels)		+ NewOut] =
+					InGainMatrix[(InputChannelIndex * InNumOutputChannels)		+ OutputChannelIndex];
+					}
+				}
+			}
+		}
+	}
+	
+	// Discrete to Discrete
+	FChannelTypeFamily::FTranscoder GetTranscoder(
+		const FDiscreteChannelTypeFamily& InFromType, const FDiscreteChannelTypeFamily& InToType, const FChannelTypeFamily::FGetTranscoderParams& InParams)
+	{
+		// Exact match? We can just memcpy each channel.
+		// TODO. in future these could be shared-ptrs from the main CAT memory block.
+		if (&InFromType == &InToType)
+		{
+			const int32 NumChannels = InToType.NumChannels();
+			return [NumChannels](TArrayView<const float*> Src, TArrayView<float*> Dst, const int32 NumFrames)
+			{
+				for (int32 i = 0; i < NumChannels; ++i)
+				{
+					FMemory::Memcpy(Dst[i], Src[i], NumFrames * sizeof(float));
+				}
+			};
+		}
+
+		switch (InParams.TranscodeMethod)
+		{
+			case EChannelTranscodeMethod::ChannelDrop:
+			{
+				return [&InToType,&InFromType](TArrayView<const float*> SrcChannels, TArrayView<float*> DstChannels, const int32 NumFrames)
+				{
+					// Copy everything destination wants, and nothing else.
+					const TArray<FDiscreteChannelTypeFamily::FSpeaker>& SrcOrder = InFromType.GetSpeakerOrder();
+					const TArray<FDiscreteChannelTypeFamily::FSpeaker>& Order = InFromType.GetSpeakerOrder();
+					const int32 NumChannels = InToType.NumChannels();	
+					for (int32 i = 0; i < NumChannels; ++i)
+					{
+						const ESpeakerShortNames DstSpeaker = Order[i].Speaker;
+						if (const int32 SrcChannelIndex = SrcOrder.IndexOfByPredicate([&](const FDiscreteChannelTypeFamily::FSpeaker& j) { return j.Speaker == DstSpeaker; });
+							SrcChannelIndex != INDEX_NONE)
+						{
+							FMemory::Memcpy(DstChannels[i], SrcChannels[SrcChannelIndex], NumFrames * sizeof(float));
+						}
+					}
+				};
+			}
+
+			case EChannelTranscodeMethod::MixUpOrDown:
+			{
+				// Make a mix matrix and call a mix/up down.
+				if (TArray<float> Gains; Create2DChannelMap(
+					{
+						.NumInputChannels = InFromType.NumChannels(),
+						.NumOutputChannels = InToType.NumChannels(),
+						.Order = EChannelMapOrder::OutputMajorOrder,
+						.MonoUpmixMethod = InParams.MixMethod,
+						.bIsCenterChannelOnly = InFromType.NumChannels() == 1 && InFromType.HasSpeaker(ESpeakerShortNames::FC)
+					}, Gains))
+				{
+					TArray<float> TranslatedGains;
+					TArray<ESpeakerShortNames> FromOrder;
+					Algo::Transform(InFromType.GetSpeakerOrder(), FromOrder, [](const FDiscreteChannelTypeFamily::FSpeaker& i) { return i.Speaker; });
+					TranslateGainMatrix(Gains, InFromType.NumChannels(), InToType.NumChannels(), GetUnrealOrder(), FromOrder, TranslatedGains);
+					return [MixGains = MoveTemp(TranslatedGains)](TArrayView<const float*> InSrc, TArrayView<float*> InDst, const int32 NumFrames)
+					{
+						MultiMonoMixUpOrDown(InSrc, InDst, NumFrames, MixGains);
+					};
+				}
+			}
+		}
+		// fail.
+		return {};
+	}
+
+	// Discrete to Ambisonics
+	FChannelTypeFamily::FTranscoder GetTranscoder(const FDiscreteChannelTypeFamily& InFromType,
+		const FAmbisonicsChannelTypeFamily& InToType, const FChannelTypeFamily::FGetTranscoderParams&)
+	{
+		const TArray<FDiscreteChannelTypeFamily::FSpeaker>& Speakers = InFromType.GetSpeakerOrder();
+		TArray<TArray<float>> AllChannelGainArrays;
+		AllChannelGainArrays.SetNum(InFromType.NumChannels());
+
+		for (int32 i = 0; i < InFromType.NumChannels(); ++i)
+		{
+			const FDiscreteChannelTypeFamily::FSpeaker& Speaker = Speakers[i];
+			
+			TArray<float>& GainArray = AllChannelGainArrays[i];
+			GainArray.SetNumZeroed(InToType.NumChannels());
+
+			// Skip LFE.
+			if (Speakers[i].Speaker == ESpeakerShortNames::LFE)
+			{
+				continue;
+			}
+			
+			//FVector2f Spherical = Speakers[i].Position.UnitCartesianToSpherical();
+			//FSphericalHarmonicCalculator::AdjustUESphericalCoordinatesForAmbisonics(Spherical);
+			const float AzimuthRads = FMath::DegreesToRadians(Speaker.AzimuthDegrees);
+			const float ElevationRads = FMath::DegreesToRadians(Speaker.ElevationDegrees);
+			FSphericalHarmonicCalculator::ComputeSoundfieldChannelGains(InToType.GetAmbisonicsOrder(), AzimuthRads, ElevationRads, GainArray);
+			FSphericalHarmonicCalculator::NormalizeGains(GainArray);
+		}
+	
+		return [Gains = MoveTemp(AllChannelGainArrays)](TArrayView<const float*> Src, TArrayView<float*> Dst, const int32 NumFrames)
+		{
+			const int32 NumSrcChannels = Src.Num();
+			for( int32 i = 0; i < Dst.Num(); ++i)
+			{
+				FMemory::Memzero(Dst[i], sizeof(float) * NumFrames);
+			}
+			for (int32 SourceChannelIndex = 0; SourceChannelIndex < NumSrcChannels; ++SourceChannelIndex)
+			{
+				EncodeMonoAmbisonicMixIn(MakeArrayView(Src[SourceChannelIndex], NumFrames), Dst, Gains[SourceChannelIndex]);
+			}
+		};
+	}
+	
+	// Ambisonics to Discrete
+	FChannelTypeFamily::FTranscoder GetTranscoder(const FAmbisonicsChannelTypeFamily& InFromType, const FDiscreteChannelTypeFamily& InToType, const FChannelTypeFamily::FGetTranscoderParams&)
+	{
+		TArray<TArray<float>> AllChannelGainArrays;
+		AllChannelGainArrays.SetNum(InToType.NumChannels());
+		const TArray<FDiscreteChannelTypeFamily::FSpeaker>& Order = InToType.GetSpeakerOrder();
+		for (int32 i = 0; i < InToType.NumChannels(); ++i)
+		{
+			const FDiscreteChannelTypeFamily::FSpeaker& Speaker = Order[i];
+			
+			TArray<float>& GainArray = AllChannelGainArrays[i];
+			GainArray.SetNum(InFromType.NumChannels());
+
+			// Skip LFE.
+			if (Order[i].Speaker == ESpeakerShortNames::LFE)
+			{
+				continue;
+			}
+			
+			//FVector2f Spherical = Order[i].Position.UnitCartesianToSpherical();
+			//FSphericalHarmonicCalculator::AdjustUESphericalCoordinatesForAmbisonics(Spherical);
+			check(GainArray.Num() == FAmbisonicsChannelTypeFamily::OrderToNumChannels(InFromType.GetAmbisonicsOrder()));
+			const float AzimuthRads = FMath::DegreesToRadians(Speaker.AzimuthDegrees);
+			const float ElevationRads = FMath::DegreesToRadians(Speaker.ElevationDegrees);
+			FSphericalHarmonicCalculator::ComputeSoundfieldChannelGains(InFromType.GetAmbisonicsOrder(), AzimuthRads, ElevationRads, GainArray);
+			FSphericalHarmonicCalculator::NormalizeGains(GainArray);
+		}
+	
+		return [NumSpeakers = InToType.NumChannels(), Gains = MoveTemp(AllChannelGainArrays)](TArrayView<const float*> Src, TArrayView<float*> Dst, const int32 NumFrames)
+		{
+			for (int32 i = 0; i < NumSpeakers; ++i)
+			{
+				TArrayView<float> DstChannel = MakeArrayView(Dst[i],NumFrames);
+				FMemory::Memzero(DstChannel.GetData(), DstChannel.NumBytes()); // Decoder is MixIn.
+				DecodeMonoAmbisonicMixIn(Src, DstChannel, Gains[i]);
+			}
+		};
+	}
+
+	// Ambisonics to Ambisonics
+	FChannelTypeFamily::FTranscoder GetTranscoder(const FAmbisonicsChannelTypeFamily& InFromType,
+		const FAmbisonicsChannelTypeFamily& InToType, const FChannelTypeFamily::FGetTranscoderParams& InParams)
+	{
+		// Ambisonics to Ambisonics
+		if (&InToType == &InFromType)
+		{
+			// Same, so memcpy.
+			return [](TArrayView<const float*> InSrc, TArrayView<float*> InDst, const int32 NumFrames)
+			{
+				for (int32 i = 0; i < InSrc.Num(); ++i)
+				{
+					FMemory::Memcpy(InDst[i], InSrc[i], sizeof(float) * NumFrames);
+				}
+			};
+		}
+
+		const int32 NumChannelsToCopy = FMath::Min(InToType.NumChannels(), InFromType.NumChannels());
+		return [NumChannelsToCopy](TArrayView<const float*> InSrc, TArrayView<float*> InDst, const int32 NumFrames)
+		{
+			// Copy as many channels as destination will hold.
+			for (int32 i = 0; i < NumChannelsToCopy; ++i)
+			{
+				FMemory::Memcpy(InDst[i], InSrc[i], sizeof(float) * NumFrames);
+			}
+		};
+	}
+}
\ No newline at end of file
diff --git a/Engine/Plugins/Experimental/MetasoundExperimental/Source/AudioExperimentalRuntime/Private/TypeFamily/ChannelTypeFamily.cpp b/Engine/Plugins/Experimental/MetasoundExperimental/Source/AudioExperimentalRuntime/Private/TypeFamily/ChannelTypeFamily.cpp
--- a/Engine/Plugins/Experimental/MetasoundExperimental/Source/AudioExperimentalRuntime/Private/TypeFamily/ChannelTypeFamily.cpp
+++ b/Engine/Plugins/Experimental/MetasoundExperimental/Source/AudioExperimentalRuntime/Private/TypeFamily/ChannelTypeFamily.cpp
@@ -2,25 +2,28 @@
 
 #include "TypeFamily/ChannelTypeFamily.h"
 
+#include "Algo/Transform.h"
 #include "CoreGlobals.h"
+#include "ChannelAgnostic/ChannelAgnosticTranscoding.h"
 #include "Containers/Map.h"
-#include "DSP/ChannelMap.h"
-#include "DSP/MultiMono.h"
+#include "DSP/SphericalHarmonicCalculator.h"
 
 namespace Audio
 {
 	namespace ChannelTypeFamilyPrviate
 	{
-		static FString MakePrettyString(const TArray<ESpeakerShortNames>& InEnums)
+		static FString MakePrettyString(const TArray<FDiscreteChannelTypeFamily::FSpeaker>& InSpeakers)
 		{
-			return FString::JoinBy(InEnums,TEXT(", "), [](ESpeakerShortNames InEnum) -> FString { return LexToString(InEnum); });
+			return FString::JoinBy(InSpeakers,TEXT(", "), [](FDiscreteChannelTypeFamily::FSpeaker InSpeaker) -> FString
+			{
+				return FString::Printf(TEXT("[Name=%s,Az=%2.2f, El=%2.2f]"), LexToString(InSpeaker.Speaker), InSpeaker.AzimuthDegrees, InSpeaker.ElevationDegrees);
+			});
 		}
 	}
 	
 	class FChannelRegistryImpl final : public IChannelTypeRegistry 
 	{
 		TMap<FName, const FChannelTypeFamily*> Types;
-		TMap<int32, const FChannelTypeFamily*> ChannelCountLookup;
 	public:
 		UE_NONCOPYABLE(FChannelRegistryImpl);
 		FChannelRegistryImpl() = default;
@@ -44,20 +47,11 @@ namespace Audio
 			
 			// Safe to cast it and add.
 			const FChannelTypeFamily* ChannelType = static_cast<const FChannelTypeFamily*>(&InType.Get());
-
-			// Log.
-			//UE_LOG(LogTemp, Display, TEXT("Registering %s, NumChannels=%d"), *ChannelType->GetName().GetPlainNameString(), ChannelType->NumChannels());
 			
 			// Add to master list of registered types.
+			check(Types.Find(InUniqueName) == nullptr);
 			Types.Add(InUniqueName,ChannelType);
-
-			// If this is a default, then add it to our Channel -> Type lookup table. Assert it's unique.
-			if (const int32 NumChannels = ChannelType->NumChannels(); NumChannels > 0 && ChannelType->IsParentsDefault())
-			{
-				const FChannelTypeFamily*& Found = ChannelCountLookup.FindOrAdd(NumChannels);
-				checkf(Found == nullptr, TEXT("%d Already registered for type %s"), NumChannels, *Found->GetName().ToString());
-				Found = ChannelType;
-			}
+	
 			return true;
 		}	
 		virtual const FTypeFamily* FindTypeInternal(const FName InUniqueName) const override
@@ -69,29 +63,6 @@ namespace Audio
 			return nullptr;
 		}
 
-		virtual const FChannelTypeFamily* FindChannelType(const int32 InNumChannels) override
-		{
-			if (const FChannelTypeFamily** Found = ChannelCountLookup.Find(InNumChannels))
-			{
-				check(!(*Found)->IsAbstract());
-				return *Found;
-			}
-
-			// If we don't have a default lookup for this channel count, look up the first thing that matches.
-			using FPair = decltype(Types)::ElementType;
-			TArray<FPair> AllTypes = Types.Array();
-			if (const FPair* Found = AllTypes.FindByPredicate([InNumChannels](const FPair& InPair) -> bool
-			{
-				return InPair.Value->NumChannels() == InNumChannels && !InPair.Value->IsAbstract();
-			}))
-			{
-				return Found->Value;	
-			}
-
-			// Failed to find anything... :(
-			return nullptr;
-		}
-
 		virtual TArray<const FChannelTypeFamily*> GetAllChannelFormats() const override
 		{
 			TArray<const FChannelTypeFamily*> AllFormats;
@@ -100,6 +71,11 @@ namespace Audio
 		};
 	};
 
+	int32 FAmbisonicsChannelTypeFamily::OrderToNumChannels(const int32 InOrder)
+	{
+		return FSphericalHarmonicCalculator::OrderToNumChannels(InOrder);
+	}
+
 	IChannelTypeRegistry& GetChannelRegistry()
 	{
 		static FChannelRegistryImpl Registry;
@@ -129,9 +105,15 @@ namespace Audio
 			InParentType->DefaultChild = this;
 		}
 	}
-	
-	
-	FDiscreteChannelTypeFamily::FDiscreteChannelTypeFamily(const FName& InUniqueName, FChannelTypeFamily* InParentType, const FString& InFriendlyName, const TArray<ESpeakerShortNames>& InOrder, const bool bIsParentsDefault, const bool bIsAbstract)
+
+   FChannelTypeFamily::FTranscoder FChannelTypeFamily::GetTranscoder(const FGetTranscoderParams& InParam) const
+   {
+		FTranscoderResolver Resolver(InParam);
+		Accept(Resolver);
+		return Resolver.MoveResult();
+   }
+
+   FDiscreteChannelTypeFamily::FDiscreteChannelTypeFamily(const FName& InUniqueName, FChannelTypeFamily* InParentType, const FString& InFriendlyName, const TArray<FSpeaker>& InOrder, const bool bIsParentsDefault, const bool bIsAbstract)
 		: Super(InUniqueName, GetFamilyTypeName(), InOrder.Num(), InParentType, InFriendlyName, bIsParentsDefault, bIsAbstract)
 		, Order(InOrder)
 	{
@@ -141,123 +123,12 @@ namespace Audio
 			*InUniqueName.ToString(),  InOrder.Num(), *GetParentNameSafe(InParentType), *InFriendlyName, ToCStr(LexToString(bIsParentsDefault)), *ChannelTypeFamilyPrviate::MakePrettyString(Order), ToCStr(LexToString(bIsAbstract)));
 	}
 
-	static TArrayView<ESpeakerShortNames> GetUnrealOrder()
-	{
-		static ESpeakerShortNames Order[]
-		{
-			#define CSV(X) X,
-			FOREACH_ENUM_ESPEAKERSHORTNAMES(CSV)
-			#undef CSV
-		};
-		return Order;
-	}
-
-	static void TranslateGainMatrix(
-		const TArrayView<float>& InGainMatrix,
-		const int32 InNumInputChannels,
-		const int32 InNumOutputChannels,
-		const TArrayView<const ESpeakerShortNames>& InOldOrder,
-		const TArrayView<const ESpeakerShortNames>& InNewOrder,
-		TArray<float>& OutGains)
-	{
-		// Channels not found in new order will be zeroed.
-		OutGains.SetNumZeroed(InNumOutputChannels*InNumInputChannels);
-		
-		// Lookup (old->new).
-		int32 Lookup[static_cast<int32>(ESpeakerShortNames::NumChannels)] = { INDEX_NONE };
-		for (int32 i = 0; i < InOldOrder.Num(); ++i)
-		{
-			Lookup[static_cast<int32>(InOldOrder[i])] = InNewOrder.IndexOfByKey(InOldOrder[i]);
-		}
-		
-		// Translate.
-		for (int32 InputChannelIndex=0; InputChannelIndex < InNumInputChannels; ++InputChannelIndex)
-		{
-			// This input channel has a mapping in the new?
-			const ESpeakerShortNames OldIn = InOldOrder[InputChannelIndex];
-			if (const int32 NewIn = Lookup[static_cast<int32>(OldIn)]; NewIn != INDEX_NONE)
-			{
-				for (int32 OutputChannelIndex=0; OutputChannelIndex < InNumOutputChannels; ++OutputChannelIndex)
-				{
-					const ESpeakerShortNames OldOut = InOldOrder[OutputChannelIndex];
-					if (const int32 NewOut = Lookup[static_cast<int32>(OldOut)]; NewOut != INDEX_NONE)
-					{
-							OutGains[(NewIn				* InNumOutputChannels)		+ NewOut] =
-						InGainMatrix[(InputChannelIndex * InNumOutputChannels)		+ OutputChannelIndex];
-					}
-				}
-			}
-		}
-	}
-	
-	FChannelTypeFamily::FTranscoder FDiscreteChannelTypeFamily::GetTranscoder(
-	const FDiscreteChannelTypeFamily& InFromType, const FGetTranscoderParams& InParams) const
-	{
-		// Exact match? We can just memcpy each channel.
-		// TODO. in future these could be shared-ptrs from the main CAT memory block.
-		if (&InFromType == this)
-		{
-			return [NumChannels = InFromType.NumChannels()](TArrayView<const float*> Src, TArrayView<float*> Dst, const int32 NumFrames)
-			{
-				for (int32 i = 0; i < NumChannels; ++i)
-				{
-					FMemory::Memcpy(Dst[i], Src[i], NumFrames * sizeof(float));
-				}
-			};
-		}
-
-		switch (InParams.TranscodeMethod)
-		{
-			case EChannelTranscodeMethod::ChannelDrop:
-			{
-				return [&](TArrayView<const float*> SrcChannels, TArrayView<float*> DstChannels, const int32 NumFrames)
-				{
-					// Copy everything destination wants, and nothing else.
-					ensure(Order.Num() == NumChannelsPrivate);
-					const TArray<ESpeakerShortNames> SrcOrder = InFromType.Order;
-					for (int32 i = 0; i < NumChannelsPrivate; ++i)
-					{
-						const ESpeakerShortNames DstSpeaker = Order[i];
-						if (const int32 SrcChannelIndex = SrcOrder.IndexOfByKey(DstSpeaker); SrcChannelIndex != INDEX_NONE)
-						{
-							FMemory::Memcpy(DstChannels[i], SrcChannels[SrcChannelIndex], NumFrames * sizeof(float));	
-						}
-					}
-				};
-			}
-			
-			case EChannelTranscodeMethod::MixUpOrDown:
-			{
-				// Make a mix matrix and call a mix/up down.
-				if (TArray<float> Gains; Create2DChannelMap(
-					{
-						.NumInputChannels = InFromType.NumChannels(),
-						.NumOutputChannels = NumChannels(),
-						.Order = EChannelMapOrder::OutputMajorOrder,
-						.MonoUpmixMethod = InParams.MixMethod,
-						.bIsCenterChannelOnly = InFromType.NumChannels() == 1 && InFromType.HasSpeaker(ESpeakerShortNames::FC)
-					}, Gains))
-				{
-					TArray<float> TranslatedGains;
-					TranslateGainMatrix(Gains, InFromType.NumChannels(), NumChannels(), GetUnrealOrder(), InFromType.GetSpeakerOrder(), TranslatedGains);
-					return [MixGains = MoveTemp(TranslatedGains)](TArrayView<const float*> InSrc, TArrayView<float*> InDst, const int32 NumFrames)
-					{
-						MultiMonoMixUpOrDown(InSrc, InDst, NumFrames, MixGains);
-					};
-				}
-			}
-		}
-
-		// fail.
-		return {};
-	}
-
 	TOptional<FChannelTypeFamily::FChannelName> FDiscreteChannelTypeFamily::GetChannelName(const int32 InChannelIndex) const
 	{
 		check(InChannelIndex >= 0 && InChannelIndex < NumChannels());
 		if (Order.IsValidIndex(InChannelIndex))
 		{
-			const ESpeakerShortNames Speaker = Order[InChannelIndex];
+			const ESpeakerShortNames Speaker = Order[InChannelIndex].Speaker;
 			const FName SpeakerName = LexToString(Speaker);
 			return
 			{ // Optional.
@@ -280,95 +151,230 @@ namespace Audio
 		// Register root type.
 		static FChannelTypeFamily BaseCat(TEXT("Cat"), TEXT("Cat"),  0, nullptr, TEXT("Base Cat"), false, true);
 		Registry.RegisterType(BaseCat.GetName(), BaseCat);
-	
+		
 #define REGISTER_CAT(NAME, PARENT_NAME, FRIENDLY_NAME, LAYOUT, DEFAULT,ABSTRACT)\
 static FDiscreteChannelTypeFamily UE_JOIN(__Type,NAME)(UE_STRINGIZE(NAME), Registry.FindChannel(PARENT_NAME), FRIENDLY_NAME, LAYOUT, DEFAULT, ABSTRACT);\
 Registry.RegisterType(TEXT(UE_STRINGIZE(NAME)),UE_JOIN(__Type,NAME))
 
 		using enum ESpeakerShortNames;
+		using FS = FDiscreteChannelTypeFamily::FSpeaker;
 
 		// Top level abstraction. "Discrete".
-		REGISTER_CAT(Discrete,			TEXT("Cat"),		TEXT("Discrete"),						{},										false,					true);
-	
-		// Mono		(UniqueName)		(Parent Name)			(Friendly Name)										(Channel membership/order.)				(is parents default).  (abstract)
-		REGISTER_CAT(Mono,				TEXT("Discrete"),TEXT("Mono"),							{},										false,					true);
-		REGISTER_CAT(Mono1Dot0,			TEXT("Mono"),	TEXT("Mono (1.0)"),					(TArray{FL}),							true,					false);
-		REGISTER_CAT(Mono1Dot1,			TEXT("Mono"),	TEXT("Mono (1.1)"),					(TArray{FL,LFE}),						false,					false);
-
-		// Stereo																								
-		REGISTER_CAT(Stereo,			TEXT("Discrete"), TEXT("Stereo"),						{},										false,					true);
-		REGISTER_CAT(Stereo2Dot0,		TEXT("Stereo"),  TEXT("Stereo (2.0)"),					(TArray{FL,FR}),						true,					false);
-		REGISTER_CAT(Stereo2Dot1,		TEXT("Stereo"),  TEXT("Stereo (2.1)"),					(TArray{FL,FR,LFE}),					false,					false);
-		REGISTER_CAT(Stereo3Dot0,		TEXT("Stereo"),  TEXT("Stereo (3.0)"),					(TArray{FL,FR,FC}),						false,					false);
-		REGISTER_CAT(Stereo3DOt1,		TEXT("Stereo"),  TEXT("Stereo (3.1)"),					(TArray{FL,FR,FC,LFE}),					false,					false);
-	
-		// Quad 
-		REGISTER_CAT(Quad,				TEXT("Discrete"),	TEXT("Quad"),						{},										false,					true);
-		REGISTER_CAT(Quad4Dot0Back,		TEXT("Quad"),		TEXT("Quad Back Speakers (4.0)"),	(TArray{FL,FR,BL,BR}),					true,					false);
-		REGISTER_CAT(Quad4Dot0Side,		TEXT("Quad"),		TEXT("Quad Side Speakers (4.0)"),	(TArray{FL,FR,SL,SR}),					false,					false);
-		REGISTER_CAT(Quad4Dot1,			TEXT("Quad"),		TEXT("Quad Back Centre LFE (4.1)"),(TArray{FL,FR,BL,BR,LFE}),				false,					false);
-
-		// Surround
-		REGISTER_CAT(Surround,			TEXT("Discrete"),	TEXT("Surround"),					{},										false,					true);
-		
-		REGISTER_CAT(Surround5,			TEXT("Surround"),	TEXT("Surround (5.X)"),			{},										true,					true);
-		REGISTER_CAT(Surround5Dot0,		TEXT("Surround5"),  TEXT("Surround (5.0)"),			(TArray{FR,FR,BL,BR,FC}),				false,					false);
-		REGISTER_CAT(Surround5_1,		TEXT("Surround5"),  TEXT("Surround (5.1)"),			(TArray{FL,FR,BL,BR,FC,LFE}),			true,					false);
+		REGISTER_CAT(Discrete,			TEXT("Cat"),		TEXT("Discrete"),		{},										false,					true);		
 
-		REGISTER_CAT(Surround7,			TEXT("Surround"),	TEXT("Surround (7.X)"),			{},										false,					true);
-		REGISTER_CAT(Surround7Dot0,		TEXT("Surround7"),  TEXT("Surround (7.0)"),			(TArray{FL,FR,SL,SR,FC,BL,BR}),			false,					false);
-		REGISTER_CAT(Surround7Dot1,		TEXT("Surround7"),  TEXT("Surround (7.1)"),			(TArray{FL,FR,SL,SR,FC,BL,BR,LFE}),		true,					false);
-		
-		// Atmos.
-		REGISTER_CAT(Atmos,				TEXT("Surround7"),		TEXT("Dolby Atmos Bed"),		{},													false,		true);
-		REGISTER_CAT(Atmos7Dot0Dot2,	TEXT("Atmos"),			TEXT("Dolby Atmos (7.0.2)"),	(TArray{FL,FR,SL,SR,FC,BL,BR,TFL,TFR}),				false,		false);
-		REGISTER_CAT(Atmos7Dot0Dot4,	TEXT("Atmos"),			TEXT("Dolby Atmos (7.0.4)"),	(TArray{FL,FR,SL,SR,FC,BL,BR,TFL,TFR,TBL,TBR}),		false,		false);
-		REGISTER_CAT(Atmos7Dot1Dot2,	TEXT("Atmos"),			TEXT("Dolby Atmos (7.1.2)"),	(TArray{FL,FR,SL,SR,FC,BL,BR,TFL,TFR,LFE}),			false,		false);
-		REGISTER_CAT(Atmos7Dot1Dot4,	TEXT("Atmos"),			TEXT("Dolby Atmos (7.1.4)"),	(TArray{FL,FR,SL,SR,FC,BL,BR,TFL,TFR,TBL,TBR,LFE}),	true,		false);
-
-#if 0
-		// Channel packs.
-		REGISTER_CAT(ChannelPack,		TEXT("Cat"),				TEXT("Channel Pack"),			{},													false,		true);
+		// Mono
+		REGISTER_CAT(Mono, TEXT("Discrete"), TEXT("Mono"), {}, false, true);
+		REGISTER_CAT(Mono1Dot0, TEXT("Mono"), TEXT("Mono (1.0)"),
+		             (TArray<FS>
+			             {
+			             { FC, 0.f, 0.f }, // Az:   0°, El:  0°
+			             }), true, false);
+
+		REGISTER_CAT(Mono1Dot1, TEXT("Mono"), TEXT("Mono (1.1)"), (TArray<FS>{
+			             { FC, 0.f, 0.f }, // Az:   0°, El:  0°
+			             { LFE, 0.f, 0.f } // LFE (omni)
+			             }), false, false);
+
+// Stereo
+REGISTER_CAT(Stereo, TEXT("Discrete"), TEXT("Stereo"), {}, false, true);
+REGISTER_CAT(Stereo2Dot0, TEXT("Stereo"), TEXT("Stereo (2.0)"),
+	(TArray<FS>{
+		{ FL, -30.f,   0.f },  // Az: -30°, El:  0°
+		{ FR,  30.f,   0.f }   // Az: +30°, El:  0°
+		}), true, false);
 		
-		// Packed mono example.
-		// Might make a new type for this, but for now....
-		static TArray<FDiscreteChannelTypeFamily> MultiMonos = []() 
-		{
-			TArray<FDiscreteChannelTypeFamily> Types;
-			constexpr int32 NumPacks = 6;
-			Types.Reserve(NumPacks);
-			for (int32 i = 0; i < NumPacks; ++i)
-			{
-				constexpr ESpeakerShortNames MonoChannel = FC;
-				TArray<ESpeakerShortNames> Order;
-				Order.Init(ESpeakerShortNames::FC, i+2); 
-				Types.Emplace(
-					*FString::Printf(TEXT("MonoPack%d"), i+2),
-					GetChannelRegistry().FindChannel(TEXT("ChannelPack")),
-					FString::Printf(TEXT("Mono Pack %d Channels"), i+2),
-					Order,	
-					false,
-					false
-				);
-				GetChannelRegistry().RegisterType(Types[i].GetName(),Types[i]);
-			}
-			return Types;
-		}();
-#endif 
+REGISTER_CAT(Stereo2Dot1, TEXT("Stereo"), TEXT("Stereo (2.1)"),
+	(TArray<FS>{
+		{ FL, -30.f, 0.f },		// Az: -30°, El:  0°
+		{ FR, 30.f, 0.f },		// Az: +30°, El:  0°
+		{ LFE, 0.f, 0.f }		// LFE (omni)
+	}),
+	false, false);
+REGISTER_CAT(Stereo3Dot0, TEXT("Stereo"), TEXT("Stereo (3.0)"),
+	(TArray<FS>{
+		{ FL, -30.f,   0.f },  // Az: -30°, El:  0°
+		{ FC,   0.f,   0.f },  // Az:   0°, El:  0°
+		{ FR,  30.f,   0.f }   // Az: +30°, El:  0°
+	}),
+	false, false);
+REGISTER_CAT(Stereo3DOt1, TEXT("Stereo"), TEXT("Stereo (3.1)"),
+	(TArray<FS>{
+		{ FL, -30.f,   0.f },  // Az: -30°, El:  0°
+		{ FC,   0.f,   0.f },  // Az:   0°, El:  0°
+		{ FR,  30.f,   0.f },  // Az: +30°, El:  0°
+		{ LFE,  0.f,   0.f }   // LFE (omni)
+	}),
+	false, false);
+
+// Quad
+REGISTER_CAT(Quad, TEXT("Discrete"), TEXT("Quad"), {}, false, true);
+REGISTER_CAT(Quad4Dot0Back, TEXT("Quad"), TEXT("Quad Back Speakers (4.0)"),
+	(TArray<FS>{
+		{ FL, -30.f,   0.f },  // Front Left
+		{ FR,  30.f,   0.f },  // Front Right
+		{ BL, -150.f,  0.f },  // Rear  Left
+		{ BR, 150.f,   0.f }   // Rear  Right
+		}),
+	true, false);
+REGISTER_CAT(Quad4Dot0Side, TEXT("Quad"), TEXT("Quad Side Speakers (4.0)"),
+	(TArray<FS>{
+		{ FL, -30.f,   0.f },  // Front Left
+		{ FR,  30.f,   0.f },  // Front Right
+		{ SL, -90.f,   0.f },  // Side  Left
+		{ SR,  90.f,   0.f }   // Side  Right
+		}),
+		false, false);
+REGISTER_CAT(Quad4Dot1, TEXT("Quad"), TEXT("Quad Back Centre LFE (4.1)"),
+	(TArray<FS>{
+		{ FL, -30.f,   0.f },  // Front Left
+		{ FR,  30.f,   0.f },  // Front Right
+		{ BL, -150.f,  0.f },  // Rear  Left
+		{ BR, 150.f,   0.f },  // Rear  Right
+		{ LFE, 0.f,    0.f }   // LFE (omni)
+		}),
+	false, false);
+
+// Surround
+REGISTER_CAT(Surround, TEXT("Discrete"), TEXT("Surround"), {}, false, true);
+REGISTER_CAT(Surround5, TEXT("Surround"), TEXT("Surround (5.X)"), {}, true, true);
+REGISTER_CAT(Surround5Dot0, TEXT("Surround5"), TEXT("Surround (5.0)"),
+    (TArray<FS>{
+        { FL,  -30,  0 },
+        { FR,   30,  0 },
+        { SL,  -90,  0 },
+        { SR,   90,  0 },
+        { FC,    0,  0 },
+        { BL, -135,  0 },
+        { BR,  135,  0 },
+    }),
+    false, false);
+
+REGISTER_CAT(Surround5_1, TEXT("Surround5"), TEXT("Surround (5.1)"),
+    (TArray<FS>{
+        { FL,  -30,  0 },
+        { FR,   30,  0 },
+        { SL,  -90,  0 },
+        { SR,   90,  0 },
+        { FC,    0,  0 },
+        { BL, -135,  0 },
+        { BR,  135,  0 },
+        { LFE,   0,  0 },
+    }),
+    true, false);
+
+REGISTER_CAT(Surround7, TEXT("Surround"), TEXT("Surround (7.X)"),
+    {}, false, true);
+
+REGISTER_CAT(Surround7Dot0, TEXT("Surround7"), TEXT("Surround (7.0)"),
+    (TArray<FS>{
+        { FL,  -30,  0 },
+        { FR,   30,  0 },
+        { SL,  -90,  0 },
+        { SR,   90,  0 },
+        { FC,    0,  0 },
+        { BL, -135,  0 },
+        { BR,  135,  0 },
+    }),
+    false, false);
+
+REGISTER_CAT(Surround7Dot1, TEXT("Surround7"), TEXT("Surround (7.1)"),
+    (TArray<FS>{
+        { FL,  -30,  0 },
+        { FR,   30,  0 },
+        { SL,  -90,  0 },
+        { SR,   90,  0 },
+        { FC,    0,  0 },
+        { LFE,   0,  0 },
+        { BL, -135,  0 },
+        { BR,  135,  0 },
+    }),
+    true, false);
+
+// Atmos beds & heights
+REGISTER_CAT(Atmos, TEXT("Surround7"), TEXT("Dolby Atmos Bed"),
+    {}, false, true);
+
+REGISTER_CAT(Atmos7Dot0Dot2, TEXT("Atmos"), TEXT("Dolby Atmos (7.0.2)"),
+    (TArray<FS>{
+        { FL,  -30.f,   0.f },
+        { FR,   30.f,   0.f },
+        { SL,  -90.f,   0.f },
+        { SR,   90.f,   0.f },
+        { FC,    0.f,   0.f },
+        { BL, -135.f,   0.f },
+        { BR,  135.f,   0.f },
+        // heights
+        { TFL, -35.3f,  30.0f },
+        { TFR,  35.3f,  30.0f },
+    }),
+    false, false);
+
+REGISTER_CAT(Atmos7Dot0Dot4, TEXT("Atmos"), TEXT("Dolby Atmos (7.0.4)"),
+    (TArray<FS>{
+        { FL,  -30.f,   0.f },
+        { FR,   30.f,   0.f },
+        { SL,  -90.f,   0.f },
+        { SR,   90.f,   0.f },
+        { FC,    0.f,   0.f },
+        { BL, -135.f,   0.f },
+        { BR,  135.f,   0.f },
+        // heights
+        { TFL, -35.3f,  30.0f },
+        { TFR,  35.3f,  30.0f },
+        { TBL, -135.0f, 45.0f },
+        { TBR,  135.0f, 45.0f },
+    }),
+    false, false);
+
+REGISTER_CAT(Atmos7Dot1Dot2, TEXT("Atmos"), TEXT("Dolby Atmos (7.1.2)"),
+    (TArray<FS>{
+        { FL,  -30.f,   0.f },
+        { FR,   30.f,   0.f },
+        { SL,  -90.f,   0.f },
+        { SR,   90.f,   0.f },
+        { FC,    0.f,   0.f },
+        { LFE,   0.f,   0.f },
+        { BL, -135.f,   0.f },
+        { BR,  135.f,   0.f },
+        // heights
+        { TFL, -35.3f,  30.0f },
+        { TFR,  35.3f,  30.0f },
+    }),
+    false, false);
+
+REGISTER_CAT(Atmos7Dot1Dot4, TEXT("Atmos"), TEXT("Dolby Atmos (7.1.4)"),
+    (TArray<FS>{
+        { FL,  -30.f,    0.f },
+        { FR,   30.f,    0.f },
+        { SL,  -90.f,    0.f },
+        { SR,   90.f,    0.f },
+        { FC,    0.f,    0.f },
+        { LFE,   0.f,    0.f },
+        { BL, -135.f,    0.f },
+        { BR,  135.f,    0.f },
+        // heights
+        { TFL, -35.3f,   30.0f },
+        { TFR,  35.3f,   30.0f },
+        { TBL, -135.0f,  45.0f },
+        { TBR,  135.0f,  45.0f },
+    }),
+    true, false);
+
 					
 #define REGISTER_AMBISONICS(NAME, PARENT_NAME, FRIENDLY_NAME, ORDER, DEFAULT,ABSTRACT)\
 static FAmbisonicsChannelTypeFamily UE_JOIN(__Type,NAME)(UE_STRINGIZE(NAME), ORDER, Registry.FindChannel(PARENT_NAME), FRIENDLY_NAME, DEFAULT, ABSTRACT);\
 ensure(Registry.RegisterType(UE_STRINGIZE(NAME),UE_JOIN(__Type,NAME)))
 
-#if 0
 		// Ambisonics 
-		REGISTER_AMBISONICS(Ambisonics,				 TEXT("Cat"), 		TEXT("Ambisonics"), 							0, false, true);
-		REGISTER_AMBISONICS(AmbisonicsFirstOrder,	 TEXT("Ambisonics"), TEXT("First Order Ambisonics (4 channels)"), 	1, false, false); 
-		REGISTER_AMBISONICS(AmbisonicsSecondOrder,	 TEXT("Ambisonics"), TEXT("Second Ambisonics (9 channels)"), 		2, false, false); 
-		REGISTER_AMBISONICS(AmbisonicsThirdOrder,	 TEXT("Ambisonics"), TEXT("Third Ambisonics (16 channels)"), 		3, false, false);
-		REGISTER_AMBISONICS(AmbisonicsFourthOrder,	 TEXT("Ambisonics"), TEXT("Fourth Ambisonics (25 chhanels)"), 		4, false, false); 
-		REGISTER_AMBISONICS(AmbisonicsFirthOrder,	 TEXT("Ambisonics"), TEXT("Firth Ambisonics (36 channels)"), 		5, false, false);  
-#endif //
+		REGISTER_AMBISONICS(Ambisonics,					TEXT("Cat"), 	   TEXT("Ambisonics"), 				0, false, true);
+		REGISTER_AMBISONICS(Ambisonics1stOrder,			TEXT("Ambisonics"), TEXT("1st Order Ambisonics"), 		1, false, false); 
+		REGISTER_AMBISONICS(Ambisonics2ndOrder,			TEXT("Ambisonics"), TEXT("2nd Order Ambisonics"), 		2, false, false); 
+		REGISTER_AMBISONICS(Ambisonics3rdOrder,			TEXT("Ambisonics"), TEXT("3rd Order Ambisonics"), 		3, false, false);
+		REGISTER_AMBISONICS(Ambisonics4thOrder,			TEXT("Ambisonics"), TEXT("4th Order Ambisonics"),		4, false, false); 
+		REGISTER_AMBISONICS(Ambisonics5thOrder,			TEXT("Ambisonics"), TEXT("5th Order Ambisonics"),		5, false, false);  
+		REGISTER_AMBISONICS(Ambisonics6thOrder,			TEXT("Ambisonics"), TEXT("6th Order Ambisonics"),		6, false, false);  
+		REGISTER_AMBISONICS(Ambisonics7thOrder,			TEXT("Ambisonics"), TEXT("7th Order Ambisonics"),		7, false, false);  
+
 		// Channel packs.
 		// FPackedChannelTypeFamily MonoPack( TEXT("PackedChannel")
 		
@@ -377,6 +383,7 @@ ensure(Registry.RegisterType(UE_STRINGIZE(NAME),UE_JOIN(__Type,NAME)))
 	}
 }
 
+
 const TCHAR* LexToString(const ESpeakerShortNames InSpeaker)
 {
 	// No code gen support here, so role this manually.
@@ -404,4 +411,4 @@ const TCHAR* LexToString(const ESpeakerShortNames InSpeaker)
 		break;
 	}
 	return nullptr;
-}
\ No newline at end of file
+}
diff --git a/Engine/Plugins/Experimental/MetasoundExperimental/Source/AudioExperimentalRuntime/Public/ChannelAgnostic/ChannelAgnosticTranscoding.h b/Engine/Plugins/Experimental/MetasoundExperimental/Source/AudioExperimentalRuntime/Public/ChannelAgnostic/ChannelAgnosticTranscoding.h
new file mode 100644
--- /dev/null
+++ b/Engine/Plugins/Experimental/MetasoundExperimental/Source/AudioExperimentalRuntime/Public/ChannelAgnostic/ChannelAgnosticTranscoding.h
@@ -0,0 +1,56 @@
+﻿// Copyright Epic Games, Inc. All Rights Reserved.
+
+#pragma once
+
+#include "TypeFamily/ChannelTypeFamily.h"
+
+namespace Audio
+{
+	// All transcoder (from -> to) combinations.
+	AUDIOEXPERIMENTALRUNTIME_API FChannelTypeFamily::FTranscoder GetTranscoder(const FDiscreteChannelTypeFamily&	InFrom,	const FDiscreteChannelTypeFamily& InTo, const FChannelTypeFamily::FGetTranscoderParams&);
+	AUDIOEXPERIMENTALRUNTIME_API FChannelTypeFamily::FTranscoder GetTranscoder(const FDiscreteChannelTypeFamily&	InFrom,	const FAmbisonicsChannelTypeFamily& InTo, const FChannelTypeFamily::FGetTranscoderParams&);
+	AUDIOEXPERIMENTALRUNTIME_API FChannelTypeFamily::FTranscoder GetTranscoder(const FAmbisonicsChannelTypeFamily& InFrom, const FDiscreteChannelTypeFamily&nTo, const FChannelTypeFamily::FGetTranscoderParams&);
+	AUDIOEXPERIMENTALRUNTIME_API FChannelTypeFamily::FTranscoder GetTranscoder(const FAmbisonicsChannelTypeFamily& InFrom, const FAmbisonicsChannelTypeFamily&nTo, const FChannelTypeFamily::FGetTranscoderParams&);
+		
+	template<typename TFromType>
+	struct TFromVisitor final : IChannelTypeVisitor
+	{
+		const TFromType& FromType;
+		const FChannelTypeFamily::FGetTranscoderParams& Params;
+		FChannelTypeFamily::FTranscoder& Transcoder;
+			
+		TFromVisitor(const TFromType& InFromType, const FChannelTypeFamily::FGetTranscoderParams& InParams, FChannelTypeFamily::FTranscoder& InTranscoder)
+			: FromType(InFromType), Params(InParams), Transcoder(InTranscoder)
+		{}
+			
+		virtual void Visit(const FDiscreteChannelTypeFamily& InTo) override
+		{
+			Transcoder = GetTranscoder(FromType, InTo, Params);	
+		}
+		virtual void Visit(const FAmbisonicsChannelTypeFamily& InTo) override
+		{
+			Transcoder = GetTranscoder(FromType, InTo, Params);
+		}
+	};
+		
+	struct FTranscoderResolver final : public IChannelTypeVisitor
+	{
+		explicit FTranscoderResolver(const FChannelTypeFamily::FGetTranscoderParams& InParams) : Params(InParams) {}
+		virtual void Visit(const FDiscreteChannelTypeFamily& InFrom) override
+		{
+			TFromVisitor Visitor(InFrom, Params, Result);
+			Params.ToType.Accept(Visitor);
+		}
+		virtual void Visit(const FAmbisonicsChannelTypeFamily& InFrom) override
+		{
+			TFromVisitor Visitor(InFrom, Params, Result);
+			Params.ToType.Accept(Visitor);
+		}
+		FChannelTypeFamily::FTranscoder&& MoveResult()
+		{
+			return MoveTemp(Result);
+		}
+		FChannelTypeFamily::FGetTranscoderParams Params;
+		FChannelTypeFamily::FTranscoder Result;
+	};
+}
\ No newline at end of file
diff --git a/Engine/Plugins/Experimental/MetasoundExperimental/Source/AudioExperimentalRuntime/Public/TypeFamily/ChannelTypeFamily.h b/Engine/Plugins/Experimental/MetasoundExperimental/Source/AudioExperimentalRuntime/Public/TypeFamily/ChannelTypeFamily.h
--- a/Engine/Plugins/Experimental/MetasoundExperimental/Source/AudioExperimentalRuntime/Public/TypeFamily/ChannelTypeFamily.h
+++ b/Engine/Plugins/Experimental/MetasoundExperimental/Source/AudioExperimentalRuntime/Public/TypeFamily/ChannelTypeFamily.h
@@ -65,7 +65,15 @@ namespace Audio
 	class FDiscreteChannelTypeFamily;
 	class FAmbisonicsChannelTypeFamily;
 	class FPackedChannelTypeFamily;
-		
+	
+	class IChannelTypeVisitor
+	{
+	public:
+		virtual ~IChannelTypeVisitor() = default;
+		virtual void Visit(const FDiscreteChannelTypeFamily&) = 0;
+		virtual void Visit(const FAmbisonicsChannelTypeFamily&) = 0;
+	};
+	
 	// Base type for all "channel" based types.
 	class FChannelTypeFamily : public FTypeFamily
 	{
@@ -143,10 +151,7 @@ namespace Audio
 		 * Get the function used for translating/transcoding between this and the "To" type.
 		 * @return The Function Object. (this will be not be set if a function wasn't returned)
 		 */
-		virtual FTranscoder GetTranscoder(const FGetTranscoderParams& InParam) const
-		{
-			return InParam.ToType.GetTranscoder(InParam);
-		}
+		AUDIOEXPERIMENTALRUNTIME_API FTranscoder GetTranscoder(const FGetTranscoderParams& InParam) const;
 
 		/**
 		 * Returns the family name (C++ type) of this ChannelType. This allows us to safely downcast (if necessary).
@@ -179,16 +184,16 @@ namespace Audio
 		template<typename T>
 		const T* Cast() const { return const_cast<FChannelTypeFamily*>(this)->Cast<T>(); }
 
+		virtual void Accept(IChannelTypeVisitor& InVisitor) const
+		{
+			checkNoEntry();
+		};
+
 	protected:
 		// Both are friends so we can hide the api.
 		friend class FDiscreteChannelTypeFamily;
 		friend class FAmbisonicsChannelTypeFamily;
 		
-		// Visitor like access for transcoder matching.
-		virtual FTranscoder GetTranscoder(const FChannelTypeFamily&, const FGetTranscoderParams&) const  { return {}; }
-		virtual FTranscoder GetTranscoder(const FDiscreteChannelTypeFamily&, const FGetTranscoderParams&) const  { return {}; }
-		virtual FTranscoder GetTranscoder(const FAmbisonicsChannelTypeFamily& , const FGetTranscoderParams&) const  { return {}; }
-		
 	private:
 		bool bIsAbstract = false;
 		bool bIsParentsDefault = false;
@@ -196,13 +201,20 @@ namespace Audio
 		int32 NumChannelsPrivate = 0;
 		FName FamilyType; 
 	};
-
-	class FDiscreteChannelTypeFamily : public FChannelTypeFamily 
+	
+	class FDiscreteChannelTypeFamily final: public FChannelTypeFamily 
 	{
 	public:
 		UE_NONCOPYABLE(FDiscreteChannelTypeFamily);
 		using Super = FChannelTypeFamily;
 
+		struct FSpeaker
+		{
+			ESpeakerShortNames Speaker;		// Enum of speaker types.
+			float AzimuthDegrees = 0.f;
+			float ElevationDegrees = 0.f;
+		};
+		
 		static FName GetFamilyTypeName()
 		{
 			static const FName Name = TEXT("Discrete");
@@ -210,7 +222,7 @@ namespace Audio
 		}
 			
 		FDiscreteChannelTypeFamily(const FName& InUniqueName, FChannelTypeFamily* InParentType, const FString& InFriendlyName,
-			const TArray<ESpeakerShortNames>& InOrder, const bool bIsParentsDefault, const bool bIsAbstract);
+			const TArray<FSpeaker>& InOrder, const bool bIsParentsDefault, const bool bIsAbstract);
 
 		/**
 		 * Find the index in the list of Channels for speaker.
@@ -219,7 +231,7 @@ namespace Audio
 		 */
 		int32 FindSpeakerIndex(const ESpeakerShortNames InSpeaker) const
 		{
-			return Order.IndexOfByKey(InSpeaker);
+			return Order.IndexOfByPredicate([InSpeaker](const FSpeaker& i) { return InSpeaker == i.Speaker; });
 		}
 
 		/**
@@ -244,21 +256,21 @@ namespace Audio
 		 * Get the order of the speakers defined on this discrete format.
 		 * @return Array of speakers in order.
 		 */
-		const TArray<ESpeakerShortNames>& GetSpeakerOrder() const { return Order; }
+		//const TArray<ESpeakerShortNames>& GetSpeakerOrder() const { return Order; }
 
-	protected:
-		// Visit "to" Type.
-		virtual FTranscoder GetTranscoder(const FGetTranscoderParams& InParams) const override
+		const TArray<FSpeaker>& GetSpeakerOrder() const { return Order; }
+
+		virtual void Accept(IChannelTypeVisitor& InVisitor) const override
 		{
-			return InParams.ToType.GetTranscoder(*this, InParams);
+			InVisitor.Visit(*this);
 		}
-		// Visitor accept.
-		virtual FTranscoder GetTranscoder(const FDiscreteChannelTypeFamily&, const FGetTranscoderParams&) const override;		
-		TArray<ESpeakerShortNames> Order;
+
+	protected:
+		TArray<FSpeaker> Order;
 	};
 
 	// Ambisonics.
-	class FAmbisonicsChannelTypeFamily : public FChannelTypeFamily 
+	class FAmbisonicsChannelTypeFamily final : public FChannelTypeFamily 
 	{
 	public:
 		UE_NONCOPYABLE(FAmbisonicsChannelTypeFamily);
@@ -270,16 +282,19 @@ namespace Audio
 			return Name;
 		}
 
-		static int32 OrderToNumChannels(const int32 InOrder)
-		{
-			// placeholder.
-			return 0;
-		}
+		static int32 OrderToNumChannels(const int32 InOrder);
 
 		explicit FAmbisonicsChannelTypeFamily(const FName& InUniqueName, const int32 InOrder, FChannelTypeFamily* InParentType, const FString& InFriendlyName, const bool bIsParentsDefault, const bool InbIsAbstract)
 			: Super(InUniqueName, GetFamilyTypeName(), OrderToNumChannels(InOrder), InParentType, InFriendlyName, bIsParentsDefault, InbIsAbstract)
 			, Order(InOrder)
 		{}
+
+		int32 GetAmbisonicsOrder() const { return Order; }
+
+		virtual void Accept(IChannelTypeVisitor& InVisitor) const override
+		{
+			InVisitor.Visit(*this);
+		}
 	private:
 		int32 Order = 0;
 	};
@@ -334,9 +349,6 @@ namespace Audio
 			return Channel;
 		}
 
-		// Find a channel type with a specific channel count. (this is best guess)
-		virtual const FChannelTypeFamily* FindChannelType(const int32 InNumChannels) = 0;
-
 		/**
 		 * Returns every registered format as an array.
 		 * @return Array of format types.
diff --git a/Engine/Source/Runtime/SignalProcessing/Private/Ambisonics.cpp b/Engine/Source/Runtime/SignalProcessing/Private/Ambisonics.cpp
new file mode 100644
--- /dev/null
+++ b/Engine/Source/Runtime/SignalProcessing/Private/Ambisonics.cpp
@@ -0,0 +1,38 @@
+// Copyright Epic Games, Inc. All Rights Reserved.
+
+#include "DSP/Ambisonics.h"
+#include "DSP/FloatArrayMath.h"
+
+namespace Audio
+{
+void EncodeMonoAmbisonicMixIn(
+	TArrayView<const float> Src, TArrayView<float*> Dst, TArrayView<const float> AmbisonicGains)
+{
+	checkSlow(Dst.Num() == AmbisonicGains.Num());
+			
+	const int32 NumFrames = Src.Num();
+	const int32 NumAmbisonicChannels = AmbisonicGains.Num();
+	const float* Gains = AmbisonicGains.GetData();
+	float** DstChannels = Dst.GetData();
+			
+	for (int32 ChannelIndex = 0; ChannelIndex < NumAmbisonicChannels; ++ChannelIndex)
+	{
+		ArrayMixIn(Src, MakeArrayView(DstChannels[ChannelIndex], NumFrames), Gains[ChannelIndex]);
+	}
+}
+void DecodeMonoAmbisonicMixIn(
+	TArrayView<const float*> Src, TArrayView<float> Dst, TArrayView<const float> AmbisonicGains)
+{
+	checkSlow(Dst.Num() == AmbisonicGains.Num());
+			
+	const int32 NumFrames = Dst.Num();
+	const int32 NumAmbisonicChannels = Src.Num();
+	const float* Gains = AmbisonicGains.GetData();
+	const float** SrcChannels = Src.GetData();
+
+	for (int32 ChannelIndex = 0; ChannelIndex < NumAmbisonicChannels; ++ChannelIndex)
+	{
+		ArrayMixIn(MakeArrayView(SrcChannels[ChannelIndex], NumFrames), Dst, Gains[ChannelIndex]);
+	}
+}
+}
\ No newline at end of file
diff --git a/Engine/Source/Runtime/SignalProcessing/Private/SphericalHarmonicCalculator.cpp b/Engine/Source/Runtime/SignalProcessing/Private/SphericalHarmonicCalculator.cpp
--- a/Engine/Source/Runtime/SignalProcessing/Private/SphericalHarmonicCalculator.cpp
+++ b/Engine/Source/Runtime/SignalProcessing/Private/SphericalHarmonicCalculator.cpp
@@ -1087,3 +1087,18 @@ void FSphericalHarmonicCalculator::GenerateFirstOrderRotationMatrixGivenDegrees(
 	constexpr const float DEG_2_RAD = PI / 180.0f;
 	GenerateFirstOrderRotationMatrixGivenRadians(RotXDegrees * DEG_2_RAD, RotYDegrees * DEG_2_RAD, RotZDegrees * DEG_2_RAD, OutMatrix);
 }
+
+void FSphericalHarmonicCalculator::NormalizeGains(TArrayView<float> Gains)
+{
+	// Normalize. gains.
+	float SumSqr = 0.f;
+	for (int32 i = 0; i < Gains.Num(); ++i)
+	{
+		SumSqr+=Gains[i]*Gains[i];
+	}
+	const float InverseLength = 1.0f / FMath::Sqrt(SumSqr);
+	for (int32 i = 0; i < Gains.Num(); ++i)
+	{
+		Gains[i]*=InverseLength;
+	}
+}
diff --git a/Engine/Source/Runtime/SignalProcessing/Public/DSP/Ambisonics.h b/Engine/Source/Runtime/SignalProcessing/Public/DSP/Ambisonics.h
new file mode 100644
--- /dev/null
+++ b/Engine/Source/Runtime/SignalProcessing/Public/DSP/Ambisonics.h
@@ -0,0 +1,11 @@
+// Copyright Epic Games, Inc. All Rights Reserved.
+
+#pragma once
+
+#include "Containers/ArrayView.h"
+
+namespace Audio
+{
+	SIGNALPROCESSING_API void EncodeMonoAmbisonicMixIn(TArrayView<const float> Src, TArrayView<float*> Dst, TArrayView<const float> AmbisonicGains);
+	SIGNALPROCESSING_API void DecodeMonoAmbisonicMixIn(TArrayView<const float*> Src, TArrayView<float> Dst, TArrayView<const float> AmbisonicGains);
+}
\ No newline at end of file
diff --git a/Engine/Source/Runtime/SignalProcessing/Public/DSP/SphericalHarmonicCalculator.h b/Engine/Source/Runtime/SignalProcessing/Public/DSP/SphericalHarmonicCalculator.h
--- a/Engine/Source/Runtime/SignalProcessing/Public/DSP/SphericalHarmonicCalculator.h
+++ b/Engine/Source/Runtime/SignalProcessing/Public/DSP/SphericalHarmonicCalculator.h
@@ -36,7 +36,7 @@ public:
 	static void AdjustUESphericalCoordinatesForAmbisonics(UE::Math::TVector2<T>& InOutVector)
 	{
 		InOutVector.X = -(InOutVector.X - HALF_PI);
-		InOutVector.Y *= -1.0f;
+		//InOutVector.Y *= -1.0f;
 		Swap(InOutVector.X,InOutVector.Y);
 	}
 	
@@ -45,4 +45,6 @@ public:
 		check(InOrder >= 0);
 		return (InOrder+1) * (InOrder+1);
 	}
+
+	static SIGNALPROCESSING_API void NormalizeGains(TArrayView<float> Gains);
 };
