Lint fixes

diff --git a/Engine/Plugins/Experimental/MetasoundExperimental/Source/AudioExperimentalRuntime/Private/ChannelAgnostic/ChannelAgnosticTranscoding.cpp b/Engine/Plugins/Experimental/MetasoundExperimental/Source/AudioExperimentalRuntime/Private/ChannelAgnostic/ChannelAgnosticTranscoding.cpp
--- a/Engine/Plugins/Experimental/MetasoundExperimental/Source/AudioExperimentalRuntime/Private/ChannelAgnostic/ChannelAgnosticTranscoding.cpp
+++ b/Engine/Plugins/Experimental/MetasoundExperimental/Source/AudioExperimentalRuntime/Private/ChannelAgnostic/ChannelAgnosticTranscoding.cpp
@@ -61,13 +61,13 @@ namespace Audio
 	
 	// Discrete to Discrete
 	FChannelTypeFamily::FTranscoder GetTranscoder(
-		const FDiscreteChannelTypeFamily& InFromType, const FDiscreteChannelTypeFamily& InToType, const FChannelTypeFamily::FGetTranscoderParams& InParams)
+		const FDiscreteChannelTypeFamily& InSrcType, const FDiscreteChannelTypeFamily& InDstType, const FChannelTypeFamily::FGetTranscoderParams& InParams)
 	{
 		// Exact match? We can just memcpy each channel.
 		// TODO. in future these could be shared-ptrs from the main CAT memory block.
-		if (&InFromType == &InToType)
+		if (&InDstType == &InSrcType)
 		{
-			const int32 NumChannels = InToType.NumChannels();
+			const int32 NumChannels = InDstType.NumChannels();
 			return [NumChannels](TArrayView<const float*> Src, TArrayView<float*> Dst, const int32 NumFrames)
 			{
 				for (int32 i = 0; i < NumChannels; ++i)
@@ -81,15 +81,15 @@ namespace Audio
 		{
 			case EChannelTranscodeMethod::ChannelDrop:
 			{
-				return [&InToType,&InFromType](TArrayView<const float*> SrcChannels, TArrayView<float*> DstChannels, const int32 NumFrames)
+				return [&InSrcType,&InDstType](TArrayView<const float*> SrcChannels, TArrayView<float*> DstChannels, const int32 NumFrames)
 				{
 					// Copy everything destination wants, and nothing else.
-					const TArray<FDiscreteChannelTypeFamily::FSpeaker>& SrcOrder = InFromType.GetSpeakerOrder();
-					const TArray<FDiscreteChannelTypeFamily::FSpeaker>& Order = InFromType.GetSpeakerOrder();
-					const int32 NumChannels = InToType.NumChannels();	
-					for (int32 i = 0; i < NumChannels; ++i)
+					const TArray<FDiscreteChannelTypeFamily::FSpeaker>& SrcOrder = InSrcType.GetSpeakerOrder();
+					const TArray<FDiscreteChannelTypeFamily::FSpeaker>& DstOrder = InDstType.GetSpeakerOrder();
+					const int32 NumDstChannels = InDstType.NumChannels();	
+					for (int32 i = 0; i < NumDstChannels; ++i)
 					{
-						const ESpeakerShortNames DstSpeaker = Order[i].Speaker;
+						const ESpeakerShortNames DstSpeaker = DstOrder[i].Speaker;
 						if (const int32 SrcChannelIndex = SrcOrder.IndexOfByPredicate([&](const FDiscreteChannelTypeFamily::FSpeaker& j) { return j.Speaker == DstSpeaker; });
 							SrcChannelIndex != INDEX_NONE)
 						{
@@ -104,17 +104,17 @@ namespace Audio
 				// Make a mix matrix and call a mix/up down.
 				if (TArray<float> Gains; Create2DChannelMap(
 					{
-						.NumInputChannels = InFromType.NumChannels(),
-						.NumOutputChannels = InToType.NumChannels(),
+						.NumInputChannels = InSrcType.NumChannels(),
+						.NumOutputChannels = InDstType.NumChannels(),
 						.Order = EChannelMapOrder::OutputMajorOrder,
 						.MonoUpmixMethod = InParams.MixMethod,
-						.bIsCenterChannelOnly = InFromType.NumChannels() == 1 && InFromType.HasSpeaker(ESpeakerShortNames::FC)
+						.bIsCenterChannelOnly = InSrcType.NumChannels() == 1 && InSrcType.HasSpeaker(ESpeakerShortNames::FC)
 					}, Gains))
 				{
 					TArray<float> TranslatedGains;
 					TArray<ESpeakerShortNames> FromOrder;
-					Algo::Transform(InFromType.GetSpeakerOrder(), FromOrder, [](const FDiscreteChannelTypeFamily::FSpeaker& i) { return i.Speaker; });
-					TranslateGainMatrix(Gains, InFromType.NumChannels(), InToType.NumChannels(), GetUnrealOrder(), FromOrder, TranslatedGains);
+					Algo::Transform(InSrcType.GetSpeakerOrder(), FromOrder, [](const FDiscreteChannelTypeFamily::FSpeaker& i) { return i.Speaker; });
+					TranslateGainMatrix(Gains, InSrcType.NumChannels(), InDstType.NumChannels(), GetUnrealOrder(), FromOrder, TranslatedGains);
 					return [MixGains = MoveTemp(TranslatedGains)](TArrayView<const float*> InSrc, TArrayView<float*> InDst, const int32 NumFrames)
 					{
 						MultiMonoMixUpOrDown(InSrc, InDst, NumFrames, MixGains);
diff --git a/Engine/Plugins/Experimental/MetasoundExperimental/Source/AudioExperimentalRuntime/Public/ChannelAgnostic/ChannelAgnosticTranscoding.h b/Engine/Plugins/Experimental/MetasoundExperimental/Source/AudioExperimentalRuntime/Public/ChannelAgnostic/ChannelAgnosticTranscoding.h
--- a/Engine/Plugins/Experimental/MetasoundExperimental/Source/AudioExperimentalRuntime/Public/ChannelAgnostic/ChannelAgnosticTranscoding.h
+++ b/Engine/Plugins/Experimental/MetasoundExperimental/Source/AudioExperimentalRuntime/Public/ChannelAgnostic/ChannelAgnosticTranscoding.h
@@ -7,10 +7,10 @@
 namespace Audio
 {
 	// All transcoder (from -> to) combinations.
-	AUDIOEXPERIMENTALRUNTIME_API FChannelTypeFamily::FTranscoder GetTranscoder(const FDiscreteChannelTypeFamily&	InFrom,	const FDiscreteChannelTypeFamily& InTo, const FChannelTypeFamily::FGetTranscoderParams&);
-	AUDIOEXPERIMENTALRUNTIME_API FChannelTypeFamily::FTranscoder GetTranscoder(const FDiscreteChannelTypeFamily&	InFrom,	const FAmbisonicsChannelTypeFamily& InTo, const FChannelTypeFamily::FGetTranscoderParams&);
-	AUDIOEXPERIMENTALRUNTIME_API FChannelTypeFamily::FTranscoder GetTranscoder(const FAmbisonicsChannelTypeFamily& InFrom, const FDiscreteChannelTypeFamily&nTo, const FChannelTypeFamily::FGetTranscoderParams&);
-	AUDIOEXPERIMENTALRUNTIME_API FChannelTypeFamily::FTranscoder GetTranscoder(const FAmbisonicsChannelTypeFamily& InFrom, const FAmbisonicsChannelTypeFamily&nTo, const FChannelTypeFamily::FGetTranscoderParams&);
+	AUDIOEXPERIMENTALRUNTIME_API FChannelTypeFamily::FTranscoder GetTranscoder(const FDiscreteChannelTypeFamily& InFrom, const FDiscreteChannelTypeFamily& InTo, const FChannelTypeFamily::FGetTranscoderParams&);
+	AUDIOEXPERIMENTALRUNTIME_API FChannelTypeFamily::FTranscoder GetTranscoder(const FDiscreteChannelTypeFamily& InFrom, const FAmbisonicsChannelTypeFamily& InTo, const FChannelTypeFamily::FGetTranscoderParams&);
+	AUDIOEXPERIMENTALRUNTIME_API FChannelTypeFamily::FTranscoder GetTranscoder(const FAmbisonicsChannelTypeFamily& InFrom, const FDiscreteChannelTypeFamily& InTo, const FChannelTypeFamily::FGetTranscoderParams&);
+	AUDIOEXPERIMENTALRUNTIME_API FChannelTypeFamily::FTranscoder GetTranscoder(const FAmbisonicsChannelTypeFamily& InFrom, const FAmbisonicsChannelTypeFamily& InTo, const FChannelTypeFamily::FGetTranscoderParams&);
 		
 	template<typename TFromType>
 	struct TFromVisitor final : IChannelTypeVisitor
diff --git a/Engine/Source/Runtime/SignalProcessing/Private/Ambisonics.cpp b/Engine/Source/Runtime/SignalProcessing/Private/Ambisonics.cpp
--- a/Engine/Source/Runtime/SignalProcessing/Private/Ambisonics.cpp
+++ b/Engine/Source/Runtime/SignalProcessing/Private/Ambisonics.cpp
@@ -23,7 +23,7 @@ void EncodeMonoAmbisonicMixIn(
 void DecodeMonoAmbisonicMixIn(
 	TArrayView<const float*> Src, TArrayView<float> Dst, TArrayView<const float> AmbisonicGains)
 {
-	checkSlow(Dst.Num() == AmbisonicGains.Num());
+	checkSlow(Src.Num() == AmbisonicGains.Num());
 			
 	const int32 NumFrames = Dst.Num();
 	const int32 NumAmbisonicChannels = Src.Num();
diff --git a/Engine/Source/Runtime/SignalProcessing/Private/SphericalHarmonicCalculator.cpp b/Engine/Source/Runtime/SignalProcessing/Private/SphericalHarmonicCalculator.cpp
--- a/Engine/Source/Runtime/SignalProcessing/Private/SphericalHarmonicCalculator.cpp
+++ b/Engine/Source/Runtime/SignalProcessing/Private/SphericalHarmonicCalculator.cpp
@@ -1096,9 +1096,12 @@ void FSphericalHarmonicCalculator::NormalizeGains(TArrayView<float> Gains)
 	{
 		SumSqr+=Gains[i]*Gains[i];
 	}
-	const float InverseLength = 1.0f / FMath::Sqrt(SumSqr);
-	for (int32 i = 0; i < Gains.Num(); ++i)
+	if (!FMath::IsNearlyZero(SumSqr))
 	{
-		Gains[i]*=InverseLength;
+		const float InverseLength = 1.0f / FMath::Sqrt(SumSqr);
+		for (int32 i = 0; i < Gains.Num(); ++i)
+		{
+			Gains[i]*=InverseLength;
+		}
 	}
 }
