Fix MetaSound node configuration details panel issues with arrays and serialization

diff --git a/Engine/Plugins/Experimental/MetasoundExperimental/Source/MetasoundExperimentalEditor/Private/MetasoundExampleNodeDetailsCustomization.cpp b/Engine/Plugins/Experimental/MetasoundExperimental/Source/MetasoundExperimentalEditor/Private/MetasoundExampleNodeDetailsCustomization.cpp
--- a/Engine/Plugins/Experimental/MetasoundExperimental/Source/MetasoundExperimentalEditor/Private/MetasoundExampleNodeDetailsCustomization.cpp
+++ b/Engine/Plugins/Experimental/MetasoundExperimental/Source/MetasoundExperimentalEditor/Private/MetasoundExampleNodeDetailsCustomization.cpp
@@ -86,15 +86,21 @@ void FExampleWidgetNodeConfigurationCustomization::OnChildPropertyChanged(const
 	{
 		if (GraphNode.IsValid() && MyFloatPropertyHandle.IsValid() && MyFloatPropertyHandle->IsValidHandle())
 		{
-			FMetaSoundFrontendDocumentBuilder& DocBuilder = GraphNode->GetBuilderChecked().GetBuilder();
+			const FMetaSoundFrontendDocumentBuilder& DocBuilder = GraphNode->GetBuilderChecked().GetBuilder();
 			const FGuid& NodeID = GraphNode->GetNodeID();
 
 			// Update the operator data value from the configuration property handle value
 			// Node configuration operator data API is experimental
 			// so this code will be made cleaner in the future 
-			TInstancedStruct<FMetaSoundFrontendNodeConfiguration> Config = DocBuilder.FindNodeConfiguration(NodeID);
+			const TConstStructView<FMetaSoundFrontendNodeConfiguration>& Config = DocBuilder.FindNodeConfiguration(NodeID);
+			if (!Config.IsValid())
+			{
+				return;
+			}
 			TSharedPtr<const Metasound::IOperatorData> OperatorData = Config.Get().GetOperatorData();
+			// Safe downcast because this operator data type is associated with the node config type associated with this customization type
 			const TSharedPtr<const Metasound::Experimental::FWidgetExampleOperatorData> WidgetOperatorData = StaticCastSharedPtr<const Metasound::Experimental::FWidgetExampleOperatorData>(OperatorData);
+			// Const cast because operator data getter is currently only const
 			TSharedPtr<Metasound::Experimental::FWidgetExampleOperatorData> MutableWidgetOperatorData = ConstCastSharedPtr<Metasound::Experimental::FWidgetExampleOperatorData>(WidgetOperatorData);
 			MyFloatPropertyHandle->GetValue(MutableWidgetOperatorData->MyFloat);
 		}
diff --git a/Engine/Plugins/Experimental/MetasoundExperimental/Source/MetasoundExperimentalRuntime/Private/MetasoundExampleNodeConfiguration.cpp b/Engine/Plugins/Experimental/MetasoundExperimental/Source/MetasoundExperimentalRuntime/Private/MetasoundExampleNodeConfiguration.cpp
--- a/Engine/Plugins/Experimental/MetasoundExperimental/Source/MetasoundExperimentalRuntime/Private/MetasoundExampleNodeConfiguration.cpp
+++ b/Engine/Plugins/Experimental/MetasoundExperimental/Source/MetasoundExperimentalRuntime/Private/MetasoundExampleNodeConfiguration.cpp
@@ -392,6 +392,7 @@ FMetaSoundWidgetExampleNodeConfiguration::FMetaSoundWidgetExampleNodeConfigurati
 
 TSharedPtr<const Metasound::IOperatorData> FMetaSoundWidgetExampleNodeConfiguration::GetOperatorData() const
 {
+	OperatorData->MyFloat = MyFloat;
 	return OperatorData;
 }
 #undef LOCTEXT_NAMESPACE // MetasoundExperimentalRuntime
diff --git a/Engine/Plugins/Experimental/MetasoundExperimental/Source/MetasoundExperimentalRuntime/Public/MetasoundExampleNodeConfiguration.h b/Engine/Plugins/Experimental/MetasoundExperimental/Source/MetasoundExperimentalRuntime/Public/MetasoundExampleNodeConfiguration.h
--- a/Engine/Plugins/Experimental/MetasoundExperimental/Source/MetasoundExperimentalRuntime/Public/MetasoundExampleNodeConfiguration.h
+++ b/Engine/Plugins/Experimental/MetasoundExperimental/Source/MetasoundExperimentalRuntime/Public/MetasoundExampleNodeConfiguration.h
@@ -64,5 +64,5 @@ struct FMetaSoundWidgetExampleNodeConfiguration : public FMetaSoundFrontendNodeC
 	TSharedPtr<const Metasound::IOperatorData> GetOperatorData() const override;
 
 private: 
-	TSharedPtr<Metasound::Experimental::FWidgetExampleOperatorData> OperatorData;
+	mutable TSharedPtr<Metasound::Experimental::FWidgetExampleOperatorData> OperatorData;
 };
diff --git a/Engine/Plugins/Runtime/Metasound/Source/MetasoundEditor/Private/MetasoundNodeConfigurationCustomization.cpp b/Engine/Plugins/Runtime/Metasound/Source/MetasoundEditor/Private/MetasoundNodeConfigurationCustomization.cpp
--- a/Engine/Plugins/Runtime/Metasound/Source/MetasoundEditor/Private/MetasoundNodeConfigurationCustomization.cpp
+++ b/Engine/Plugins/Runtime/Metasound/Source/MetasoundEditor/Private/MetasoundNodeConfigurationCustomization.cpp
@@ -19,6 +19,7 @@ namespace Metasound::Editor
 		{
 			TDelegate<void(const FPropertyChangedEvent&)> OnValueChangedDelegate = TDelegate<void(const FPropertyChangedEvent&)>::CreateSP(this, &FMetaSoundNodeConfigurationDataDetails::OnChildPropertyChanged);
 			ChildHandle->SetOnPropertyValueChangedWithData(OnValueChangedDelegate);
+			ChildHandle->SetOnChildPropertyValueChangedWithData(OnValueChangedDelegate);
 		}
 	}
 
