adding new experiment engine runtime module
adding mapping function node which uses node config
adding a clock node which returns the current running time of a metasound

#rnx

diff --git a/Engine/Plugins/Experimental/MetasoundExperimental/MetasoundExperimental.uplugin b/Engine/Plugins/Experimental/MetasoundExperimental/MetasoundExperimental.uplugin
--- a/Engine/Plugins/Experimental/MetasoundExperimental/MetasoundExperimental.uplugin
+++ b/Engine/Plugins/Experimental/MetasoundExperimental/MetasoundExperimental.uplugin
@@ -25,6 +25,11 @@
 			"Type": "Runtime",
 			"LoadingPhase": "Default"
 		},
+		{
+			"Name": "MetasoundExperimentalEngineRuntime",
+			"Type": "Runtime",
+			"LoadingPhase": "Default"
+		},
 		{
 			"Name": "MetasoundExperimentalEditor",
 			"Type": "Editor",
diff --git a/Engine/Plugins/Experimental/MetasoundExperimental/Source/MetasoundExperimentalEditor/MetasoundExperimentalEditor.build.cs b/Engine/Plugins/Experimental/MetasoundExperimental/Source/MetasoundExperimentalEditor/MetasoundExperimentalEditor.build.cs
--- a/Engine/Plugins/Experimental/MetasoundExperimental/Source/MetasoundExperimentalEditor/MetasoundExperimentalEditor.build.cs
+++ b/Engine/Plugins/Experimental/MetasoundExperimental/Source/MetasoundExperimentalEditor/MetasoundExperimentalEditor.build.cs
@@ -38,6 +38,7 @@ public class MetasoundExperimentalEditor : ModuleRules
 				"CoreUObject",
 				"Engine",
 				"MetasoundExperimentalRuntime",
+				"MetasoundExperimentalEngineRuntime",
 				"PropertyEditor",
 				"UnrealEd",
 				"Slate",
diff --git a/Engine/Plugins/Experimental/MetasoundExperimental/Source/MetasoundExperimentalEditor/Private/MetasoundExperimentalEditorModule.cpp b/Engine/Plugins/Experimental/MetasoundExperimental/Source/MetasoundExperimentalEditor/Private/MetasoundExperimentalEditorModule.cpp
--- a/Engine/Plugins/Experimental/MetasoundExperimental/Source/MetasoundExperimentalEditor/Private/MetasoundExperimentalEditorModule.cpp
+++ b/Engine/Plugins/Experimental/MetasoundExperimental/Source/MetasoundExperimentalEditor/Private/MetasoundExperimentalEditorModule.cpp
@@ -3,6 +3,8 @@
 #include "MetasoundEditorModule.h"
 #include "MetasoundExampleNodeConfiguration.h"
 #include "MetasoundExampleNodeDetailsCustomization.h"
+#include "MetasoundMappingFunctionNode.h"
+#include "MetasoundMappingFunctionDetailsCustomization.h"
 #include "Modules/ModuleManager.h"
 
 #define LOCTEXT_NAMESPACE "MetasoundExperimentalEditor"
@@ -25,12 +27,19 @@ void FMetasoundExperimentalEditorModule::StartupModule()
 		{
 			return MakeShareable(new FExampleWidgetNodeConfigurationCustomization(InStructProperty, InNode));
 		});
+	MetasoundEditorModule.RegisterCustomNodeConfigurationDetailsCustomization(
+		FMetaSoundMappingFunctionNodeConfiguration::StaticStruct()->GetFName(),
+		[](TSharedPtr<IPropertyHandle> InStructProperty, TWeakObjectPtr<UMetasoundEditorGraphNode> InNode)
+		{
+			return MakeShareable(new FMappingFunctionNodeConfigurationCustomization(InStructProperty, InNode));
+		});
 }
 
 void FMetasoundExperimentalEditorModule::ShutdownModule()
 {
 	Metasound::Editor::IMetasoundEditorModule& MetasoundEditorModule = FModuleManager::GetModuleChecked<Metasound::Editor::IMetasoundEditorModule>("MetasoundEditor");
 	MetasoundEditorModule.UnregisterCustomNodeConfigurationDetailsCustomization(FMetaSoundWidgetExampleNodeConfiguration::StaticStruct()->GetFName());
+	MetasoundEditorModule.UnregisterCustomNodeConfigurationDetailsCustomization(FMetaSoundMappingFunctionNodeConfiguration::StaticStruct()->GetFName());
 }
 
 #undef LOCTEXT_NAMESPACE
diff --git a/Engine/Plugins/Experimental/MetasoundExperimental/Source/MetasoundExperimentalEditor/Private/MetasoundMappingFunctionDetailsCustomization.cpp b/Engine/Plugins/Experimental/MetasoundExperimental/Source/MetasoundExperimentalEditor/Private/MetasoundMappingFunctionDetailsCustomization.cpp
new file mode 100644
--- /dev/null
+++ b/Engine/Plugins/Experimental/MetasoundExperimental/Source/MetasoundExperimentalEditor/Private/MetasoundMappingFunctionDetailsCustomization.cpp
@@ -0,0 +1,330 @@
+// Copyright Epic Games, Inc. All Rights Reserved.
+
+#include "MetasoundMappingFunctionDetailsCustomization.h" 
+
+#include "DetailWidgetRow.h"
+#include "IDetailChildrenBuilder.h"
+#include "InstancedStructDetails.h"
+#include "MetasoundEditorGraphNode.h"
+#include "MetasoundBuilderBase.h"
+#include "MetasoundMappingFunctionNode.h"
+#include "PropertyEditorModule.h"
+#include "Widgets/Input/SSlider.h"
+#include "Widgets/Layout/SBox.h"
+#include "SCurveEditor.h"
+
+#define LOCTEXT_NAMESPACE "MetasoundExperimentalEditor"
+
+const FString StructIdentifierWithDelimiters =  TEXT(".Struct.");
+
+FMappingFunctionNodeConfigurationCustomization::FMappingFunctionNodeConfigurationCustomization(TSharedPtr<IPropertyHandle> InStructProperty, TWeakObjectPtr<UMetasoundEditorGraphNode> InNode)
+	: Metasound::Editor::FMetaSoundNodeConfigurationDataDetails(InStructProperty, InNode)
+{
+	if (InStructProperty && InStructProperty->IsValidHandle())
+	{
+		StructPropertyPath = InStructProperty->GeneratePathToProperty();
+	}
+}
+
+TArray<FRichCurveEditInfo> FMappingFunctionNodeConfigurationCustomization::GetCurves()
+{
+	TArray<FRichCurveEditInfo> Curves;
+	if (RuntimeCurve)
+	{
+		// Provide the FRichCurve for editing
+		Curves.Add(FRichCurveEditInfo(RuntimeCurve->GetRichCurve()));
+	}
+	return Curves;
+}
+
+TArray<FRichCurveEditInfoConst> FMappingFunctionNodeConfigurationCustomization::GetCurves() const
+{
+	TArray<FRichCurveEditInfoConst> Curves;
+	if (RuntimeCurve)
+	{
+		// Provide the FRichCurve for editing
+		Curves.Add(FRichCurveEditInfoConst(RuntimeCurve->GetRichCurve()));
+	}
+	return Curves;
+}
+
+void FMappingFunctionNodeConfigurationCustomization::GetCurves(TAdderReserverRef<FRichCurveEditInfoConst> Curves) const 
+{
+	if (RuntimeCurve)
+	{
+		Curves.Add(FRichCurveEditInfoConst(RuntimeCurve->GetRichCurveConst()));
+	}
+}
+
+
+bool FMappingFunctionNodeConfigurationCustomization::IsValidCurve(FRichCurveEditInfo CurveInfo)
+{
+	// Validate that the curve being edited is the one we expect
+	return RuntimeCurve && (CurveInfo.CurveToEdit == RuntimeCurve->GetRichCurve());
+}
+
+void FMappingFunctionNodeConfigurationCustomization::MakeTransactional()
+{
+	// Mark owners as transactional to support undo/redo
+	for (UObject* Owner : OwnerObjects)
+	{
+		if (Owner)
+		{
+			Owner->SetFlags(RF_Transactional);
+		}
+	}
+}
+
+void FMappingFunctionNodeConfigurationCustomization::ModifyOwner()
+{
+	// Called at the start of a curve edit (begin transaction)
+	for (UObject* Owner : OwnerObjects)
+	{
+		if (Owner)
+		{
+			Owner->Modify();
+		}
+	}
+}
+
+void FMappingFunctionNodeConfigurationCustomization::ModifyOwnerChange()
+{
+	// Called during interactive changes (e.g., while dragging a key)
+	for (UObject* Owner : OwnerObjects)
+	{
+		if (Owner)
+		{
+			Owner->Modify();
+		}
+	}
+}
+
+void FMappingFunctionNodeConfigurationCustomization::OnCurveChanged(const TArray<FRichCurveEditInfo>& ChangedCurveEditInfos)
+{
+	// Enforce clamping and endpoint constraints whenever the curve is edited
+	FRichCurve* RichCurve = RuntimeCurve ? RuntimeCurve->GetRichCurve() : nullptr;
+	if (!RichCurve)
+	{
+		return;
+	}
+
+	bool bCurveModified = false;
+
+	// Clamp all keys within [0,1] for both X (Time) and Y (Value)
+	for (FRichCurveKey& Key : RichCurve->Keys)
+	{
+		if (Key.Time < 0.0f) { Key.Time = 0.0f;  bCurveModified = true; }
+		if (Key.Time > 1.0f) { Key.Time = 1.0f;  bCurveModified = true; }
+		if (Key.Value < 0.0f) { Key.Value = 0.0f; bCurveModified = true; }
+		if (Key.Value > 1.0f) { Key.Value = 1.0f; bCurveModified = true; }
+	}
+
+	// Ensure a key exists at X = 0.0 (lock first key's X position to 0)
+	if (RichCurve->Keys.Num() == 0 || RichCurve->Keys[0].Time > 0.0f)
+	{
+		float NewValue = (RichCurve->Keys.Num() > 0) ? RichCurve->Eval(0.0f) : 0.0f;
+		RichCurve->AddKey(0.0f, NewValue);
+		RichCurve->Keys[0].InterpMode = RCIM_Linear;
+		bCurveModified = true;
+	}
+	// Ensure a key exists at X = 1.0 (lock last key's X position to 1)
+	if (RichCurve->Keys.Num() == 0 || RichCurve->Keys.Last().Time < 1.0f)
+	{
+		float NewValue = (RichCurve->Keys.Num() > 0) ? RichCurve->Eval(1.0f) : 1.0f;
+		RichCurve->AddKey(1.0f, NewValue);
+		RichCurve->Keys.Last().InterpMode = RCIM_Linear;
+		bCurveModified = true;
+	}
+
+	// Snap the first and last keys exactly to 0.0 and 1.0 on X (in case they were moved)
+	if (RichCurve->Keys.Num() >= 2)
+	{
+		FRichCurveKey& FirstKey = RichCurve->Keys[0];
+		FRichCurveKey& LastKey = RichCurve->Keys.Last();
+		if (!FMath::IsNearlyZero(FirstKey.Time))
+		{
+			FirstKey.Time = 0.0f;
+			bCurveModified = true;
+		}
+		if (!FMath::IsNearlyEqual(LastKey.Time, 1.0f))
+		{
+			LastKey.Time = 1.0f;
+			bCurveModified = true;
+		}
+	}
+
+	if (bCurveModified)
+	{
+		// Notify the engine that the property value has changed, to support undo/redo and UI refresh
+		FProperty* CurveProperty = CurvePropertyHandle.IsValid() ? CurvePropertyHandle->GetProperty() : nullptr;
+		if (CurveProperty)
+		{
+			for (UObject* Owner : OwnerObjects)
+			{
+				if (Owner)
+				{
+					Owner->PreEditChange(CurveProperty);
+				}
+			}
+			FPropertyChangedEvent ChangeEvent(CurveProperty, EPropertyChangeType::ValueSet);
+			for (UObject* Owner : OwnerObjects)
+			{
+				if (Owner)
+				{
+					Owner->PostEditChangeProperty(ChangeEvent);
+				}
+			}
+		}
+	}
+
+	UpdateMappingFunctionData();
+
+	return;
+}
+
+TArray<const UObject*> FMappingFunctionNodeConfigurationCustomization::GetOwners() const
+{
+	TArray<const UObject*> Result;
+	for (UObject* Owner : OwnerObjects)
+	{
+		if (Owner)
+		{
+			Result.Add(Owner);
+		}
+	}
+	return Result;
+}
+
+void FMappingFunctionNodeConfigurationCustomization::OnChildRowAdded(IDetailPropertyRow& ChildRow)
+{
+	TSharedPtr<IPropertyHandle> ChildHandle = ChildRow.GetPropertyHandle();
+	if (!ChildHandle || !ChildHandle->IsValidHandle())
+	{
+		return;
+	}
+
+	const FString PropertyPath = ChildHandle->GeneratePathToProperty();
+	
+	if (PropertyPath == StructPropertyPath + StructIdentifierWithDelimiters + GET_MEMBER_NAME_CHECKED(FMetaSoundMappingFunctionNodeConfiguration, MappingFunction).ToString())
+ 	{
+ 		CurvePropertyHandle = ChildHandle;
+		RuntimeCurve = nullptr;
+		OwnerObjects.Empty();
+
+		// Access the raw struct data and cast to a curve
+		TArray<void*> RawData;
+		CurvePropertyHandle->AccessRawData(RawData);
+		if (RawData.Num() > 0)
+		{
+			RuntimeCurve = static_cast<FRuntimeFloatCurve*>(RawData[0]);
+		}
+		if (!RuntimeCurve)
+		{
+			return;
+		}
+
+		// Get outer owning UObject(s)
+		CurvePropertyHandle->GetOuterObjects(OwnerObjects);
+
+		// Ensure default keys at 0 and 1 with linear mapping if the curve is empty
+		FRichCurve* RichCurve = RuntimeCurve->GetRichCurve();
+		if (RichCurve->GetNumKeys() == 0)
+		{
+			RichCurve->AddKey(0.0f, 0.0f);
+			RichCurve->AddKey(1.0f, 1.0f);
+			
+			// Set interpolation to linear for a straight line between (0,0) and (1,1)
+			RichCurve->Keys[0].InterpMode = RCIM_Linear;
+			RichCurve->Keys.Last().InterpMode = RCIM_Linear;
+		}
+
+		CurveEditorWidget = SNew(SCurveEditor)
+			.ViewMinInput(0.0f)
+			.ViewMaxInput(1.0f)
+			.ViewMinOutput(0.0f)
+			.ViewMaxOutput(1.0f)
+			.ZoomToFitHorizontal(false)
+			.ZoomToFitVertical(false)
+			.ShowInputGridNumbers(false)
+			.ShowOutputGridNumbers(false)
+			.AllowZoomOutput(false)
+			.TimelineLength(1.0f)
+			.HideUI(true)
+			.ShowZoomButtons(false)
+			.DesiredSize(FVector2D(300,200))
+			.ShowCurveSelector(false);
+
+		CurveEditorWidget->SetCurveOwner(this);
+
+		TSharedPtr<SWidget> NameWidget;
+		TSharedPtr<SWidget> ValueWidget;
+		FDetailWidgetRow Row;
+
+		ChildRow.GetDefaultWidgets(NameWidget, ValueWidget, Row);
+		ChildRow.CustomWidget(true)
+			.NameContent()
+			[
+				NameWidget.ToSharedRef()
+			]
+			.ValueContent()
+			.MinDesiredWidth(200)
+			.MaxDesiredWidth(400)
+			[
+				CurveEditorWidget.ToSharedRef()
+			];
+ 	}
+	else if (PropertyPath == StructPropertyPath + StructIdentifierWithDelimiters + GET_MEMBER_NAME_CHECKED(FMetaSoundMappingFunctionNodeConfiguration, bWrapInputs).ToString())
+	{
+		bWrapInputsPropertyHandle = ChildHandle;
+	}
+
+	// Add custom onvalue changed
+	TDelegate<void(const FPropertyChangedEvent&)> OnValueChangedDelegate = TDelegate<void(const FPropertyChangedEvent&)>::CreateSP(this, &FMappingFunctionNodeConfigurationCustomization::OnChildPropertyChanged);
+	ChildHandle->SetOnPropertyValueChangedWithData(OnValueChangedDelegate);
+
+	// Add base class on value changed
+	Metasound::Editor::FMetaSoundNodeConfigurationDataDetails::OnChildRowAdded(ChildRow);
+
+	// Make sure we update the mapping function data since we added some default keys
+	UpdateMappingFunctionData();
+}
+
+void FMappingFunctionNodeConfigurationCustomization::UpdateMappingFunctionData()
+{
+	if (GraphNode.IsValid())
+	{
+		FMetaSoundFrontendDocumentBuilder& DocBuilder = GraphNode->GetBuilderChecked().GetBuilder();
+		const FGuid& NodeID = GraphNode->GetNodeID();
+
+		TInstancedStruct<FMetaSoundFrontendNodeConfiguration> Config = DocBuilder.FindNodeConfiguration(NodeID);
+		TSharedPtr<const Metasound::IOperatorData> OperatorData = Config.Get().GetOperatorData();
+		const TSharedPtr<const Metasound::Experimental::FMappingFunctionNodeOperatorData> MappingFunctionOperatorData = StaticCastSharedPtr<const Metasound::Experimental::FMappingFunctionNodeOperatorData>(OperatorData);
+		TSharedPtr<Metasound::Experimental::FMappingFunctionNodeOperatorData> MutableMappingFunctionOperatorData = ConstCastSharedPtr<Metasound::Experimental::FMappingFunctionNodeOperatorData>(MappingFunctionOperatorData);
+
+		if (CurvePropertyHandle.IsValid() &&
+			CurvePropertyHandle->IsValidHandle())
+		{
+			TArray<void*> RawData;
+			CurvePropertyHandle->AccessRawData(RawData);
+			if (RawData.Num() > 0)
+			{
+				FRuntimeFloatCurve* Curve = static_cast<FRuntimeFloatCurve*>(RawData[0]);
+				MutableMappingFunctionOperatorData->MappingFunction = *Curve;
+			}
+		}
+		else if (bWrapInputsPropertyHandle.IsValid() &&
+				 bWrapInputsPropertyHandle->IsValidHandle())
+		{
+			bWrapInputsPropertyHandle->GetValue(MutableMappingFunctionOperatorData->bWrapInputs);
+		}
+	}
+}
+
+
+void FMappingFunctionNodeConfigurationCustomization::OnChildPropertyChanged(const FPropertyChangedEvent& InPropertyChangedEvent)
+{
+	UpdateMappingFunctionData();
+
+}
+
+#undef LOCTEXT_NAMESPACE
diff --git a/Engine/Plugins/Experimental/MetasoundExperimental/Source/MetasoundExperimentalEditor/Private/MetasoundMappingFunctionDetailsCustomization.h b/Engine/Plugins/Experimental/MetasoundExperimental/Source/MetasoundExperimentalEditor/Private/MetasoundMappingFunctionDetailsCustomization.h
new file mode 100644
--- /dev/null
+++ b/Engine/Plugins/Experimental/MetasoundExperimental/Source/MetasoundExperimentalEditor/Private/MetasoundMappingFunctionDetailsCustomization.h
@@ -0,0 +1,39 @@
+// Copyright Epic Games, Inc. All Rights Reserved.
+#pragma once
+
+#include "MetasoundNodeConfigurationCustomization.h"
+#include "SCurveEditor.h"
+#include "Curves/CurveFloat.h"
+
+class FMappingFunctionNodeConfigurationCustomization : public Metasound::Editor::FMetaSoundNodeConfigurationDataDetails, public FCurveOwnerInterface
+{
+public:
+	FMappingFunctionNodeConfigurationCustomization(TSharedPtr<IPropertyHandle> InStructProperty, TWeakObjectPtr<UMetasoundEditorGraphNode> InNode);
+
+	virtual void OnChildRowAdded(IDetailPropertyRow& ChildRow) override;
+
+	// FCurveOwnerInterface interface
+	virtual TArray<FRichCurveEditInfo> GetCurves() override;
+	virtual void GetCurves(TAdderReserverRef<FRichCurveEditInfoConst> Curves) const override;
+	virtual TArray<FRichCurveEditInfoConst> GetCurves() const override;
+	virtual bool HasRichCurves() const override { return true; }
+	virtual bool IsLinearColorCurve() const override { return false; }
+	virtual bool IsValidCurve(FRichCurveEditInfo CurveInfo) override;
+	virtual void MakeTransactional() override;
+	virtual void ModifyOwner() override;
+	virtual void ModifyOwnerChange() override;
+	virtual void OnCurveChanged(const TArray<FRichCurveEditInfo>& ChangedCurveEditInfos) override;
+	virtual TArray<const UObject*> GetOwners() const override;
+
+private:
+	void OnChildPropertyChanged(const FPropertyChangedEvent& InPropertyChangedEvent);
+	void UpdateMappingFunctionData();
+
+	FString StructPropertyPath;
+	TSharedPtr<IPropertyHandle> CurvePropertyHandle;
+	TSharedPtr<IPropertyHandle> bWrapInputsPropertyHandle;
+
+	FRuntimeFloatCurve* RuntimeCurve = nullptr;
+	TArray<UObject*> OwnerObjects;
+	TSharedPtr<SCurveEditor> CurveEditorWidget;
+};
diff --git a/Engine/Plugins/Experimental/MetasoundExperimental/Source/MetasoundExperimentalEngineRuntime/MetasoundExperimentalEngineRuntime.Build.cs b/Engine/Plugins/Experimental/MetasoundExperimental/Source/MetasoundExperimentalEngineRuntime/MetasoundExperimentalEngineRuntime.Build.cs
new file mode 100644
--- /dev/null
+++ b/Engine/Plugins/Experimental/MetasoundExperimental/Source/MetasoundExperimentalEngineRuntime/MetasoundExperimentalEngineRuntime.Build.cs
@@ -0,0 +1,41 @@
+// Copyright Epic Games, Inc. All Rights Reserved.
+
+using UnrealBuildTool;
+
+public class MetasoundExperimentalEngineRuntime : ModuleRules
+{
+	public MetasoundExperimentalEngineRuntime(ReadOnlyTargetRules Target) : base(Target)
+	{
+		PCHUsage = ModuleRules.PCHUsageMode.UseExplicitOrSharedPCHs;
+		PublicDependencyModuleNames.AddRange(
+			new string[]
+			{
+				"Core",
+				"MetasoundGraphCore",
+				"MetasoundFrontend",
+				"AudioExperimentalRuntime",
+				"MetasoundExperimentalRuntime"
+			}
+		);
+		
+		// When the Editor builds all plugins, it will build with engine. In that case add a dependency on CoreUObject.
+		// If we don't do this we will get linker errors for the Engine expecting to find stuff defined here.
+		if (Target.bCompileAgainstEngine)
+		{
+			PublicDependencyModuleNames.Add("CoreUObject");
+		}
+		
+		PrivateDependencyModuleNames.AddRange(
+			new string[]
+			{
+				"Core",
+				"Engine",
+				"MetasoundGraphCore",
+				"MetasoundFrontend",
+				"MetasoundEngine",
+				"MetasoundStandardNodes",
+				"SignalProcessing"
+			}
+		);
+	}
+}
diff --git a/Engine/Plugins/Experimental/MetasoundExperimental/Source/MetasoundExperimentalEngineRuntime/Private/MetasoundExperimentalEngineRuntimeModule.cpp b/Engine/Plugins/Experimental/MetasoundExperimental/Source/MetasoundExperimentalEngineRuntime/Private/MetasoundExperimentalEngineRuntimeModule.cpp
new file mode 100644
--- /dev/null
+++ b/Engine/Plugins/Experimental/MetasoundExperimental/Source/MetasoundExperimentalEngineRuntime/Private/MetasoundExperimentalEngineRuntimeModule.cpp
@@ -0,0 +1,22 @@
+// Copyright Epic Games, Inc. All Rights Reserved.
+
+#include "MetasoundExperimentalEngineRuntimeModule.h"
+
+#include "MetasoundDataTypeRegistrationMacro.h"
+#include "Misc/Paths.h"
+#include "Modules/ModuleManager.h"
+
+void FMetasoundExperimentalEngineRuntimeModule::StartupModule()
+{
+	FModuleManager::Get().LoadModule("AudioExperimentalRuntime");
+	FMetasoundFrontendRegistryContainer::Get()->RegisterPendingNodes();
+}
+
+void FMetasoundExperimentalEngineRuntimeModule::ShutdownModule()
+{
+}
+IMPLEMENT_MODULE(FMetasoundExperimentalEngineRuntimeModule, MetasoundExperimentalEngineRuntime);
+
+//#if WITH_ENGINE
+//#error ("Being linked with Engine")
+//#endif //
diff --git a/Engine/Plugins/Experimental/MetasoundExperimental/Source/MetasoundExperimentalEngineRuntime/Private/MetasoundExperimentalEngineRuntimeModule.h b/Engine/Plugins/Experimental/MetasoundExperimental/Source/MetasoundExperimentalEngineRuntime/Private/MetasoundExperimentalEngineRuntimeModule.h
new file mode 100644
--- /dev/null
+++ b/Engine/Plugins/Experimental/MetasoundExperimental/Source/MetasoundExperimentalEngineRuntime/Private/MetasoundExperimentalEngineRuntimeModule.h
@@ -0,0 +1,17 @@
+// Copyright Epic Games, Inc. All Rights Reserved.
+
+#pragma once
+
+#include "Modules/ModuleInterface.h"
+
+
+class FMetasoundExperimentalEngineRuntimeModule : public IModuleInterface
+{ 
+public:
+	//~ Begin IModuleInterface API
+	virtual void StartupModule() override;
+	virtual void ShutdownModule() override;
+	//~ End IModuleInterface API
+
+private:
+};
diff --git a/Engine/Plugins/Experimental/MetasoundExperimental/Source/MetasoundExperimentalEngineRuntime/Private/MetasoundMappingFunctionNode.cpp b/Engine/Plugins/Experimental/MetasoundExperimental/Source/MetasoundExperimentalEngineRuntime/Private/MetasoundMappingFunctionNode.cpp
new file mode 100644
--- /dev/null
+++ b/Engine/Plugins/Experimental/MetasoundExperimental/Source/MetasoundExperimentalEngineRuntime/Private/MetasoundMappingFunctionNode.cpp
@@ -0,0 +1,183 @@
+// Copyright Epic Games, Inc. All Rights Reserved.
+
+#include "MetasoundMappingFunctionNode.h" 
+
+#include "Internationalization/Text.h"
+#include "Math/UnrealMathUtility.h"
+
+#include "MetasoundDataFactory.h"
+#include "MetasoundFacade.h"
+#include "MetasoundOperatorData.h"
+#include "MetasoundNodeRegistrationMacro.h"
+#include "MetasoundExecutableOperator.h"
+#include "MetasoundParamHelper.h"
+#include "MetasoundPrimitives.h"
+#include "MetasoundTrigger.h"
+
+#define LOCTEXT_NAMESPACE "MetasoundMappingFunctionNodeRuntime"
+
+namespace Metasound::Experimental
+{
+	namespace MappingFunctionNodePrivate
+	{
+		METASOUND_PARAM(InputValue, "Input", "Input float value");
+		METASOUND_PARAM(InMin, "In Min", "Input Min Value");
+		METASOUND_PARAM(InMax, "In Max", "Input Max Value");
+		METASOUND_PARAM(InOutMin, "Out Min", "Output Min Value");
+		METASOUND_PARAM(InOutMax, "Out Max", "Output Max Value");
+
+		METASOUND_PARAM(OutputFloat, "Output", "Mapped float output");
+
+		FVertexInterface GetVertexInterface()
+		{
+			FInputVertexInterface InputInterface;
+			InputInterface.Add(TInputDataVertex<float>{METASOUND_GET_PARAM_NAME_AND_METADATA(InputValue), 0.0f});
+			InputInterface.Add(TInputDataVertex<float>{METASOUND_GET_PARAM_NAME_AND_METADATA(InMin), 0.0f});
+			InputInterface.Add(TInputDataVertex<float>{METASOUND_GET_PARAM_NAME_AND_METADATA(InMax), 1.0f});
+			InputInterface.Add(TInputDataVertex<float>{METASOUND_GET_PARAM_NAME_AND_METADATA(InOutMin), 0.0f});
+			InputInterface.Add(TInputDataVertex<float>{METASOUND_GET_PARAM_NAME_AND_METADATA(InOutMax), 1.0f});
+
+			FOutputVertexInterface OutputInterface;
+			OutputInterface.Add(TOutputDataVertex<float>{METASOUND_GET_PARAM_NAME_AND_METADATA(OutputFloat)});
+
+			return FVertexInterface
+			{
+				MoveTemp(InputInterface),
+				MoveTemp(OutputInterface)
+			};
+		}
+	}
+
+	const FLazyName FMappingFunctionNodeOperatorData::OperatorDataTypeName = "MappingFunctionOperatorData";
+
+	class FMappingFunctionNodeOperator : public TExecutableOperator<FMappingFunctionNodeOperator>
+	{
+	public:
+		FMappingFunctionNodeOperator(const TSharedPtr<const FMappingFunctionNodeOperatorData>& InOperatorData, const FBuildOperatorParams& InParams)
+			: OperatorData(InOperatorData)
+			, InMin(InParams.InputData.GetOrCreateDefaultDataReadReference<float>(MappingFunctionNodePrivate::InMinName, InParams.OperatorSettings))
+			, InMax(InParams.InputData.GetOrCreateDefaultDataReadReference<float>(MappingFunctionNodePrivate::InMaxName, InParams.OperatorSettings))
+			, InOutMin(InParams.InputData.GetOrCreateDefaultDataReadReference<float>(MappingFunctionNodePrivate::InOutMinName, InParams.OperatorSettings))
+			, InOutMax(InParams.InputData.GetOrCreateDefaultDataReadReference<float>(MappingFunctionNodePrivate::InOutMaxName, InParams.OperatorSettings))
+			, InputValue(InParams.InputData.GetOrCreateDefaultDataReadReference<float>(MappingFunctionNodePrivate::InputValueName, InParams.OperatorSettings))
+			, OutputValue(TDataWriteReferenceFactory<float>::CreateExplicitArgs(InParams.OperatorSettings))
+		{
+		}
+
+		virtual void BindInputs(FInputVertexInterfaceData& InOutVertexData) override
+		{
+			using namespace MappingFunctionNodePrivate;
+
+			InOutVertexData.BindReadVertex(METASOUND_GET_PARAM_NAME(InputValue), InputValue);
+			InOutVertexData.BindReadVertex(METASOUND_GET_PARAM_NAME(InMin), InMin);
+			InOutVertexData.BindReadVertex(METASOUND_GET_PARAM_NAME(InMax), InMax);
+			InOutVertexData.BindReadVertex(METASOUND_GET_PARAM_NAME(InOutMin), InOutMin);
+			InOutVertexData.BindReadVertex(METASOUND_GET_PARAM_NAME(InOutMax), InOutMax);
+		}
+
+		virtual void BindOutputs(FOutputVertexInterfaceData& InOutVertexData) override
+		{
+			using namespace MappingFunctionNodePrivate;
+			InOutVertexData.BindWriteVertex(METASOUND_GET_PARAM_NAME(OutputFloat), OutputValue);
+		}
+
+		void Execute()
+		{
+			if (!OperatorData.IsValid())
+			{
+				return;
+			}
+
+			const FRuntimeFloatCurve& Curve = OperatorData->MappingFunction;
+			const FRichCurve& RichCurve = *Curve.GetRichCurveConst();
+			if (!RichCurve.GetNumKeys())
+			{
+				return;
+			}
+
+			float MinFuncTime = 0.0f;
+			float MaxFuncTime = 0.0f;
+			float MinFuncValue = 0.0f;
+			float MaxFuncValue = 0.0f;
+
+			RichCurve.GetTimeRange(MinFuncTime, MaxFuncTime);
+			RichCurve.GetValueRange(MinFuncValue, MaxFuncValue);
+
+			float TimeInputVal = 0.0f;
+			if (OperatorData->bWrapInputs)
+			{
+				TimeInputVal = FMath::Wrap(*InputValue, *InMin, *InMax);
+			}
+			else
+			{
+				TimeInputVal = FMath::Clamp(*InputValue, *InMin, *InMax);
+			}
+
+			float TimeInputMapped = FMath::GetMappedRangeValueClamped(FVector2D{ *InMin, *InMax}, FVector2D{ MinFuncTime, MaxFuncTime }, TimeInputVal);
+			float OutputVal = RichCurve.Eval(TimeInputMapped);
+			float OutputValMapped = FMath::GetMappedRangeValueClamped(FVector2D{ MinFuncValue, MaxFuncValue }, FVector2D{ *InOutMin, *InOutMax}, OutputVal);
+			
+			*OutputValue = OutputValMapped;
+		}
+		
+		static FNodeClassMetadata GetNodeInfo()
+		{
+			using namespace MappingFunctionNodePrivate;
+			return FNodeClassMetadata
+			{
+				FNodeClassName{ "Experimental", "MappingFunctionOperator", "" },
+				1, // Major version
+				0, // Minor version
+				LOCTEXT("MappingFunctionNodeName", "Mapping Function"),
+				LOCTEXT("MappingFunctionNodeDescription", "A node which maps inputs to outputs using a curve mapping function."),
+				TEXT("UE"), // Author
+				LOCTEXT("MappingFunctionPromptIfMissing", "Enable the MetaSoundExperimental Plugin"), // Prompt if missing
+				GetVertexInterface(),
+				{}
+			};
+		}
+
+		static TUniquePtr<IOperator> CreateOperator(const FBuildOperatorParams& InParams, FBuildResults& OutResults)
+		{
+			using namespace MappingFunctionNodePrivate;
+
+			if (const FMappingFunctionNodeOperatorData* OperatorData = CastOperatorData<const FMappingFunctionNodeOperatorData>(InParams.Node.GetOperatorData().Get()))
+			{
+				// Need to get a shared ptr of the operator data so we can get live updates from the interface
+				const TSharedPtr<const FMappingFunctionNodeOperatorData>& OperatorDataSharedPtr = StaticCastSharedPtr<const FMappingFunctionNodeOperatorData>(InParams.Node.GetOperatorData());
+				return MakeUnique<FMappingFunctionNodeOperator>(OperatorDataSharedPtr, InParams);
+			}
+
+			return nullptr;
+		}
+
+	private:
+		// Contains configured float data
+		TSharedPtr<const FMappingFunctionNodeOperatorData> OperatorData;
+
+		TDataReadReference<float> InMin;
+		TDataReadReference<float> InMax;
+		TDataReadReference<float> InOutMin;
+		TDataReadReference<float> InOutMax;
+
+		TDataReadReference<float> InputValue;
+		TDataWriteReference<float> OutputValue;
+	};
+
+	using FMappingFunctionNode = TNodeFacade<FMappingFunctionNodeOperator>;
+	METASOUND_REGISTER_NODE_AND_CONFIGURATION(FMappingFunctionNode, FMetaSoundMappingFunctionNodeConfiguration);
+}
+
+FMetaSoundMappingFunctionNodeConfiguration::FMetaSoundMappingFunctionNodeConfiguration()
+	: OperatorData(MakeShared<Metasound::Experimental::FMappingFunctionNodeOperatorData>(MappingFunction, bWrapInputs))
+{
+}
+
+TSharedPtr<const Metasound::IOperatorData> FMetaSoundMappingFunctionNodeConfiguration::GetOperatorData() const
+{
+	OperatorData->MappingFunction = MappingFunction;
+	OperatorData->bWrapInputs = bWrapInputs;
+	return OperatorData;
+}
+
+#undef LOCTEXT_NAMESPACE // MetasoundMappingFunctionNodeRuntime
\ No newline at end of file
diff --git a/Engine/Plugins/Experimental/MetasoundExperimental/Source/MetasoundExperimentalEngineRuntime/Private/MetasoundTimerNode.cpp b/Engine/Plugins/Experimental/MetasoundExperimental/Source/MetasoundExperimentalEngineRuntime/Private/MetasoundTimerNode.cpp
new file mode 100644
--- /dev/null
+++ b/Engine/Plugins/Experimental/MetasoundExperimental/Source/MetasoundExperimentalEngineRuntime/Private/MetasoundTimerNode.cpp
@@ -0,0 +1,168 @@
+// Copyright Epic Games, Inc. All Rights Reserved.
+
+#include "Internationalization/Text.h"
+#include "Math/UnrealMathUtility.h"
+
+#include "MetasoundDataFactory.h"
+#include "MetasoundFacade.h"
+#include "MetasoundOperatorData.h"
+#include "MetasoundNodeRegistrationMacro.h"
+#include "MetasoundExecutableOperator.h"
+#include "MetasoundParamHelper.h"
+#include "MetasoundPrimitives.h"
+#include "MetasoundTrigger.h"
+
+#define LOCTEXT_NAMESPACE "MetasoundTimerNodeRuntime"
+
+namespace Metasound::Experimental
+{
+	namespace TimerNodePrivate
+	{
+		METASOUND_PARAM(InStart, "Start", "Starts the Timer output from 0.0.");
+		METASOUND_PARAM(InStop, "Stop", "Stops the Timer output");
+		METASOUND_PARAM(InScale, "Scale", "Scales the output time of the Timer");
+
+		METASOUND_PARAM(OutputTime, "Timer Output", "Time in seconds since Timer started.");
+
+		FVertexInterface GetVertexInterface()
+		{
+			FInputVertexInterface InputInterface;
+			InputInterface.Add(TInputDataVertex<FTrigger>{METASOUND_GET_PARAM_NAME_AND_METADATA(InStart)});
+			InputInterface.Add(TInputDataVertex<FTrigger>{METASOUND_GET_PARAM_NAME_AND_METADATA(InStop)});
+			InputInterface.Add(TInputDataVertex<float>{METASOUND_GET_PARAM_NAME_AND_METADATA(InScale), 1.0f});
+
+			FOutputVertexInterface OutputInterface;
+			OutputInterface.Add(TOutputDataVertex<float>{METASOUND_GET_PARAM_NAME_AND_METADATA(OutputTime)});
+
+			return FVertexInterface
+			{
+				MoveTemp(InputInterface),
+				MoveTemp(OutputInterface)
+			};
+		}
+	}
+
+	class FTimerNodeOperator : public TExecutableOperator<FTimerNodeOperator>
+	{
+	public:
+		FTimerNodeOperator(const FBuildOperatorParams& InParams)
+			: InStartTrigger(InParams.InputData.GetOrCreateDefaultDataReadReference<FTrigger>(TimerNodePrivate::InStartName, InParams.OperatorSettings))
+			, InStopTrigger(InParams.InputData.GetOrCreateDefaultDataReadReference<FTrigger>(TimerNodePrivate::InStopName, InParams.OperatorSettings))
+			, InScale(InParams.InputData.GetOrCreateDefaultDataReadReference<float>(TimerNodePrivate::InScaleName, InParams.OperatorSettings))
+			, OutputTime(TDataWriteReferenceFactory<float>::CreateExplicitArgs(InParams.OperatorSettings))
+			, FramesPerSecond(InParams.OperatorSettings.GetSampleRate())
+			, FramesPerBlock(InParams.OperatorSettings.GetNumFramesPerBlock())
+		{
+		}
+
+		virtual void BindInputs(FInputVertexInterfaceData& InOutVertexData) override
+		{
+			using namespace TimerNodePrivate;
+
+			InOutVertexData.BindReadVertex(METASOUND_GET_PARAM_NAME(InStart), InStartTrigger);
+			InOutVertexData.BindReadVertex(METASOUND_GET_PARAM_NAME(InStop), InStopTrigger);
+			InOutVertexData.BindReadVertex(METASOUND_GET_PARAM_NAME(InScale), InScale);
+		}
+
+		virtual void BindOutputs(FOutputVertexInterfaceData& InOutVertexData) override
+		{
+			using namespace TimerNodePrivate;
+			InOutVertexData.BindWriteVertex(METASOUND_GET_PARAM_NAME(OutputTime), OutputTime);
+		}
+
+		void Execute()
+		{
+			float ClampedScale = *InScale;
+			if (FMath::IsNearlyZero(*InScale, UE_KINDA_SMALL_NUMBER))
+			{
+				ClampedScale = UE_KINDA_SMALL_NUMBER;
+			}
+			ScaledFramesPerSecond = FramesPerSecond / ClampedScale;
+			check(ScaledFramesPerSecond > 0.0f || ScaledFramesPerSecond < 0.0f);
+			*OutputTime = (float)FramesSinceStart / ScaledFramesPerSecond;
+
+			// In a given block, if we start and stop a bunch of times, we only want to look at the last start and stops
+			const TArray<int32>& StartTriggerFrames = InStartTrigger->GetTriggeredFrames();
+			const TArray<int32>& StopTriggerFrames = InStopTrigger->GetTriggeredFrames();
+
+			int32 LastStartTriggerFrame = StartTriggerFrames.Num() > 0 ? StartTriggerFrames.Last() : INDEX_NONE;
+			int32 LastStopTriggerFrame = StopTriggerFrames.Num() > 0 ? StopTriggerFrames.Last() : INDEX_NONE;
+
+			// Start and stop were both triggered in this block
+			if (LastStartTriggerFrame != INDEX_NONE && LastStopTriggerFrame != INDEX_NONE)
+			{
+				// If we start then stopped as the last triggers, then our "time" value is the delta between these frames
+				if (LastStopTriggerFrame > LastStartTriggerFrame)
+				{
+					// We track the number of frames since start as this time delta so next block it can be reported in output
+					FramesSinceStart = LastStopTriggerFrame - LastStartTriggerFrame;
+					bIsTimerRunning = false;
+				}
+				else
+				{
+					// We track the remaining frames to the end of the block 
+					FramesSinceStart = FramesPerBlock - LastStartTriggerFrame;
+					bIsTimerRunning = true;
+				}
+			}
+			// Only stop was triggered, so accumulate up to the stop trigger frame
+			else if (LastStopTriggerFrame != INDEX_NONE)
+			{
+				FramesSinceStart += LastStopTriggerFrame;
+				bIsTimerRunning = false;
+			}
+			// Only start was triggered, so track from the start frame to the end of the block
+			else if (LastStartTriggerFrame != INDEX_NONE)
+			{
+				FramesSinceStart = FramesPerBlock - LastStartTriggerFrame;
+				bIsTimerRunning = true;
+			}
+			// No stop or start trigger was made, so simply accumulate the time if we're running
+			// otherwise, don't need to do anything.
+			else if (bIsTimerRunning)
+			{
+				FramesSinceStart += FramesPerBlock;
+			}
+		}
+		
+		static FNodeClassMetadata GetNodeInfo()
+		{
+			using namespace TimerNodePrivate;
+			return FNodeClassMetadata
+			{
+				FNodeClassName{ "Experimental", "TimerNodeOperator", "" },
+				1, // Major version
+				0, // Minor version
+				LOCTEXT("TimerNodeName", "Timer"),
+				LOCTEXT("TimerNodeDescription", "A node which outputs a time value in seconds with optional scaling."),
+				TEXT("UE"), // Author
+				LOCTEXT("TimernodePromptIfMissing", "Enable the MetaSoundExperimental Plugin"), // Prompt if missing
+				GetVertexInterface(),
+				{}
+			};
+		}
+
+		static TUniquePtr<IOperator> CreateOperator(const FBuildOperatorParams& InParams, FBuildResults& OutResults)
+		{
+			return MakeUnique<FTimerNodeOperator>(InParams);
+		}
+
+	private:
+		TDataReadReference<FTrigger> InStartTrigger;
+		TDataReadReference<FTrigger> InStopTrigger;
+		TDataReadReference<float> InScale;
+
+		TDataWriteReference<float> OutputTime;
+
+		float FramesPerSecond = 0.0f;
+		float ScaledFramesPerSecond = 0.0f;
+		int32 FramesSinceStart = 0;
+		int32 FramesPerBlock = 0;
+		bool bIsTimerRunning = false;
+	};
+
+	using FTimerNode = TNodeFacade<FTimerNodeOperator>;
+	METASOUND_REGISTER_NODE(FTimerNode);
+}
+
+#undef LOCTEXT_NAMESPACE // MetasoundTimerNodeRuntime
\ No newline at end of file
diff --git a/Engine/Plugins/Experimental/MetasoundExperimental/Source/MetasoundExperimentalEngineRuntime/Public/MetasoundMappingFunctionNode.h b/Engine/Plugins/Experimental/MetasoundExperimental/Source/MetasoundExperimentalEngineRuntime/Public/MetasoundMappingFunctionNode.h
new file mode 100644
--- /dev/null
+++ b/Engine/Plugins/Experimental/MetasoundExperimental/Source/MetasoundExperimentalEngineRuntime/Public/MetasoundMappingFunctionNode.h
@@ -0,0 +1,49 @@
+// Copyright Epic Games, Inc. All Rights Reserved.
+
+#pragma once
+
+#include "MetasoundOperatorData.h"
+#include "MetasoundFrontendDocument.h"
+#include "Curves/CurveFloat.h"
+
+#include "MetasoundMappingFunctionNode.generated.h"
+
+
+namespace Metasound::Experimental
+{
+	class FMappingFunctionNodeOperatorData : public TOperatorData<FMappingFunctionNodeOperatorData>
+	{
+	public:
+		static const FLazyName OperatorDataTypeName;
+
+		FMappingFunctionNodeOperatorData(const FRuntimeFloatCurve& InMappingFunction, bool bInWrapInputs)
+			: MappingFunction(InMappingFunction)
+			, bWrapInputs(bInWrapInputs)
+		{
+		}
+
+		FRuntimeFloatCurve MappingFunction;
+		bool bWrapInputs = false;
+	};
+}
+
+USTRUCT()
+struct FMetaSoundMappingFunctionNodeConfiguration : public FMetaSoundFrontendNodeConfiguration
+{
+	GENERATED_BODY()
+
+	FMetaSoundMappingFunctionNodeConfiguration();
+
+	// Whether or not to wrap the input values to remain in the defined input range. Otherwise, inputs are clamped.
+	UPROPERTY(EditAnywhere, Category = General)
+	bool bWrapInputs = false;
+
+	// Mapping function (curve) to map inputs to outputs
+	UPROPERTY(EditAnywhere, Category = General)
+	FRuntimeFloatCurve MappingFunction;
+
+	TSharedPtr<const Metasound::IOperatorData> GetOperatorData() const override;
+
+private: 
+	TSharedPtr<Metasound::Experimental::FMappingFunctionNodeOperatorData> OperatorData;
+};
