Fixing compile error

#rnx

diff --git a/Engine/Plugins/Experimental/MetasoundExperimental/Source/MetasoundExperimentalEditor/Private/MetasoundMappingFunctionDetailsCustomization.cpp b/Engine/Plugins/Experimental/MetasoundExperimental/Source/MetasoundExperimentalEditor/Private/MetasoundMappingFunctionDetailsCustomization.cpp
--- a/Engine/Plugins/Experimental/MetasoundExperimental/Source/MetasoundExperimentalEditor/Private/MetasoundMappingFunctionDetailsCustomization.cpp
+++ b/Engine/Plugins/Experimental/MetasoundExperimental/Source/MetasoundExperimentalEditor/Private/MetasoundMappingFunctionDetailsCustomization.cpp
@@ -15,8 +15,6 @@
 
 #define LOCTEXT_NAMESPACE "MetasoundExperimentalEditor"
 
-const FString StructIdentifierWithDelimiters =  TEXT(".Struct.");
-
 FMappingFunctionNodeConfigurationCustomization::FMappingFunctionNodeConfigurationCustomization(TSharedPtr<IPropertyHandle> InStructProperty, TWeakObjectPtr<UMetasoundEditorGraphNode> InNode)
 	: Metasound::Editor::FMetaSoundNodeConfigurationDataDetails(InStructProperty, InNode)
 {
@@ -204,10 +202,10 @@ void FMappingFunctionNodeConfigurationCustomization::OnChildRowAdded(IDetailProp
 	}
 
 	const FString PropertyPath = ChildHandle->GeneratePathToProperty();
-	
-	if (PropertyPath == StructPropertyPath + StructIdentifierWithDelimiters + GET_MEMBER_NAME_CHECKED(FMetaSoundMappingFunctionNodeConfiguration, MappingFunction).ToString())
- 	{
- 		CurvePropertyHandle = ChildHandle;
+
+	if (PropertyPath == StructPropertyPath + TEXT(".Struct.") + GET_MEMBER_NAME_CHECKED(FMetaSoundMappingFunctionNodeConfiguration, MappingFunction).ToString())
+	{
+		CurvePropertyHandle = ChildHandle;
 		RuntimeCurve = nullptr;
 		OwnerObjects.Empty();
 
@@ -232,7 +230,7 @@ void FMappingFunctionNodeConfigurationCustomization::OnChildRowAdded(IDetailProp
 		{
 			RichCurve->AddKey(0.0f, 0.0f);
 			RichCurve->AddKey(1.0f, 1.0f);
-			
+
 			// Set interpolation to linear for a straight line between (0,0) and (1,1)
 			RichCurve->Keys[0].InterpMode = RCIM_Linear;
 			RichCurve->Keys.Last().InterpMode = RCIM_Linear;
@@ -272,8 +270,8 @@ void FMappingFunctionNodeConfigurationCustomization::OnChildRowAdded(IDetailProp
 			[
 				CurveEditorWidget.ToSharedRef()
 			];
- 	}
-	else if (PropertyPath == StructPropertyPath + StructIdentifierWithDelimiters + GET_MEMBER_NAME_CHECKED(FMetaSoundMappingFunctionNodeConfiguration, bWrapInputs).ToString())
+	}
+	else if (PropertyPath == StructPropertyPath + TEXT(".Struct.") + GET_MEMBER_NAME_CHECKED(FMetaSoundMappingFunctionNodeConfiguration, bWrapInputs).ToString())
 	{
 		bWrapInputsPropertyHandle = ChildHandle;
 	}
