[Metasound Experimental] Addressed minor optimization feedback

#rnx

diff --git a/Engine/Plugins/Experimental/MetasoundExperimental/Source/MetasoundExperimentalEngineRuntime/Private/MetasoundFadeNode.cpp b/Engine/Plugins/Experimental/MetasoundExperimental/Source/MetasoundExperimentalEngineRuntime/Private/MetasoundFadeNode.cpp
--- a/Engine/Plugins/Experimental/MetasoundExperimental/Source/MetasoundExperimentalEngineRuntime/Private/MetasoundFadeNode.cpp
+++ b/Engine/Plugins/Experimental/MetasoundExperimental/Source/MetasoundExperimentalEngineRuntime/Private/MetasoundFadeNode.cpp
@@ -276,7 +276,7 @@ namespace Metasound
 
 				// Setup temporary buffer
 				TArray<float> Value;
-				Value.SetNumZeroed(ActualSamplesToGenerate + SamplesRemaining);
+				Value.SetNum(ActualSamplesToGenerate + SamplesRemaining);
 
 				// Get View of total scratch buffer
 				TArrayView<float> ValueView(Value);
