Metasound LOCTEXT fixups

#rnx

diff --git a/Engine/Plugins/Experimental/MetasoundExperimental/Source/MetasoundExperimentalRuntime/Private/MetasoundCatBreakNode.cpp b/Engine/Plugins/Experimental/MetasoundExperimental/Source/MetasoundExperimentalRuntime/Private/MetasoundCatBreakNode.cpp
--- a/Engine/Plugins/Experimental/MetasoundExperimental/Source/MetasoundExperimentalRuntime/Private/MetasoundCatBreakNode.cpp
+++ b/Engine/Plugins/Experimental/MetasoundExperimental/Source/MetasoundExperimentalRuntime/Private/MetasoundCatBreakNode.cpp
@@ -35,11 +35,9 @@ namespace Metasound
 		TOutputDataVertex<FAudioBuffer> MakeOutputDataVertex(const int32 InIndex, const FName InName, const FString& InFriendlyName)
 		{
 			const FName Name = MakeOutputVertexName(InIndex);
-#if WITH_EDITOR
-			const FText DisplayName = FText::Format(LOCTEXT("In_DisplayName", "{0}"), FText::FromString(InName.ToString()));
-#else
-			const FText DisplayName = FText::GetEmpty();
-#endif
+
+			const FText DisplayName = METASOUND_LOCTEXT_FORMAT("In_DisplayName", "{0}", FText::FromString(InName.ToString()));
+
 			return TOutputDataVertex<FAudioBuffer>{Name, FDataVertexMetadata{FText::FromString(*InFriendlyName), DisplayName}};
 		}
 		
@@ -238,10 +236,10 @@ namespace Metasound
 				FNodeClassName{ "Experimental", "CatBreakOperator", "" },
 				1, // Major version
 				0, // Minor version
-				LOCTEXT("CatBreakNodeName", "CAT Break Node"),	
-				LOCTEXT("CatBreakNodeNameDescription", "A Node that Breaks CATs"),	
+				METASOUND_LOCTEXT("CatBreakNodeName", "CAT Break Node"),	
+				METASOUND_LOCTEXT("CatBreakNodeNameDescription", "A Node that Breaks CATs"),	
 				TEXT("UE"), // Author
-				LOCTEXT("ExampleConfigurablePromptIfMissing", "Enable the MetaSoundExperimental Plugin"), // Prompt if missing
+				METASOUND_LOCTEXT("ExampleConfigurablePromptIfMissing", "Enable the MetaSoundExperimental Plugin"), // Prompt if missing
 				GetDefaultInterface(),
 				{}
 			};
diff --git a/Engine/Plugins/Experimental/MetasoundExperimental/Source/MetasoundExperimentalRuntime/Private/MetasoundCatCastingNode.cpp b/Engine/Plugins/Experimental/MetasoundExperimental/Source/MetasoundExperimentalRuntime/Private/MetasoundCatCastingNode.cpp
--- a/Engine/Plugins/Experimental/MetasoundExperimental/Source/MetasoundExperimentalRuntime/Private/MetasoundCatCastingNode.cpp
+++ b/Engine/Plugins/Experimental/MetasoundExperimental/Source/MetasoundExperimentalRuntime/Private/MetasoundCatCastingNode.cpp
@@ -156,10 +156,10 @@ namespace Metasound
 				FNodeClassName{ "Experimental", "CatCastingOperator", "" },
 				1, // Major version
 				0, // Minor version
-				LOCTEXT("CatCastingNodeName", "CAT Casting Node"),	
-				LOCTEXT("CatCastingNodeNameDescription", "A Node that allows Casting to CATs"),	
+				METASOUND_LOCTEXT("CatCastingNodeName", "CAT Casting Node"),	
+				METASOUND_LOCTEXT("CatCastingNodeNameDescription", "A Node that allows Casting to CATs"),	
 				TEXT("UE"), // Author
-				LOCTEXT("ExampleConfigurablePromptIfMissing", "Enable the MetaSoundExperimental Plugin"), // Prompt if missing
+				METASOUND_LOCTEXT("ExampleConfigurablePromptIfMissing", "Enable the MetaSoundExperimental Plugin"), // Prompt if missing
 				GetDefaultInterface(),
 				{}
 			};
diff --git a/Engine/Plugins/Experimental/MetasoundExperimental/Source/MetasoundExperimentalRuntime/Private/MetasoundCatMakeNode.cpp b/Engine/Plugins/Experimental/MetasoundExperimental/Source/MetasoundExperimentalRuntime/Private/MetasoundCatMakeNode.cpp
--- a/Engine/Plugins/Experimental/MetasoundExperimental/Source/MetasoundExperimentalRuntime/Private/MetasoundCatMakeNode.cpp
+++ b/Engine/Plugins/Experimental/MetasoundExperimental/Source/MetasoundExperimentalRuntime/Private/MetasoundCatMakeNode.cpp
@@ -34,11 +34,9 @@ namespace Metasound
 		FInputDataVertex MakeInputDataVertex(const int32 InIndex, const FName InName, const FString& InFriendlyName)
 		{
 			const FName InputName = MakeInputVertexName(InIndex);
-#if WITH_EDITOR
-			const FText InputDisplayName = FText::Format(LOCTEXT("In_DisplayName", "{0}"), FText::FromString(InName.ToString()));
-#else
-			const FText InputDisplayName = FText::GetEmpty();
-#endif
+
+			const FText InputDisplayName = METASOUND_LOCTEXT_FORMAT("In_DisplayName", "{0}", FText::FromString(InName.ToString()));
+
 			return TInputDataVertex<FAudioBuffer>{InputName, FDataVertexMetadata{FText::FromString(*InFriendlyName), InputDisplayName}};
 		}
 		
@@ -189,10 +187,10 @@ namespace Metasound
 				FNodeClassName{ "Experimental", "CatMakeOperator", "" },
 				1, // Major version
 				0, // Minor version
-				LOCTEXT("CatMakeNodeName", "CAT Make Node"),	
-				LOCTEXT("CatMakeNodeNameDescription", "A Node that builds CATs"),	
+				METASOUND_LOCTEXT("CatMakeNodeName", "CAT Make Node"),	
+				METASOUND_LOCTEXT("CatMakeNodeNameDescription", "A Node that builds CATs"),	
 				TEXT("UE"), // Author
-				LOCTEXT("ExampleConfigurablePromptIfMissing", "Enable the MetaSoundExperimental Plugin"), // Prompt if missing
+				METASOUND_LOCTEXT("ExampleConfigurablePromptIfMissing", "Enable the MetaSoundExperimental Plugin"), // Prompt if missing
 				GetDefaultInterface(),
 				{}
 			};
diff --git a/Engine/Plugins/Experimental/MetasoundExperimental/Source/MetasoundExperimentalRuntime/Private/MetasoundCatMixerNode.cpp b/Engine/Plugins/Experimental/MetasoundExperimental/Source/MetasoundExperimentalRuntime/Private/MetasoundCatMixerNode.cpp
--- a/Engine/Plugins/Experimental/MetasoundExperimental/Source/MetasoundExperimentalRuntime/Private/MetasoundCatMixerNode.cpp
+++ b/Engine/Plugins/Experimental/MetasoundExperimental/Source/MetasoundExperimentalRuntime/Private/MetasoundCatMixerNode.cpp
@@ -60,22 +60,14 @@ namespace Metasound
 		static FInputDataVertex MakeInputCatVertex(const int32 InIndex, const FName InName, const FString& InFriendlyName)
 		{
 			const FName InputName = MakeCatInputVertexName(InIndex);
-#if WITH_EDITOR
-			const FText InputDisplayName = FText::Format(LOCTEXT("In_DisplayName", "{0}"), FText::FromString(FString::Printf(TEXT("%s %d"), *InName.ToString(), InIndex)));
-#else
-			const FText InputDisplayName = FText::GetEmpty();
-#endif
+			const FText InputDisplayName = METASOUND_LOCTEXT_FORMAT("In_DisplayName", "{0}", FText::FromString(FString::Printf(TEXT("%s %d"), *InName.ToString(), InIndex)));
 			return TInputDataVertex<FChannelAgnosticType>{InputName, FDataVertexMetadata{FText::FromString(*InFriendlyName), InputDisplayName}};
 		}
 
 		static FInputDataVertex MakeInputGainVertex(const int32 InIndex, const FName InName, const FString& InFriendlyName)
 		{
 			const FName InputName = MakeGainInputVertexName(InIndex);
-#if WITH_EDITOR
-			const FText InputDisplayName = FText::Format(LOCTEXT("In_DisplayName", "{0}"), FText::FromString(FString::Printf(TEXT("%s %d"), *InName.ToString(), InIndex)));
-#else
-			const FText InputDisplayName = FText::GetEmpty();
-#endif
+			const FText InputDisplayName = METASOUND_LOCTEXT_FORMAT("In_DisplayName", "{0}", FText::FromString(FString::Printf(TEXT("%s %d"), *InName.ToString(), InIndex)));
 			return TInputDataVertex<float>{InputName, FDataVertexMetadata{ FText::FromString(*InFriendlyName), InputDisplayName }, 1.0f};
 		}
 
diff --git a/Engine/Plugins/Experimental/MetasoundExperimental/Source/MetasoundExperimentalRuntime/Private/MetasoundExampleNodeConfiguration.cpp b/Engine/Plugins/Experimental/MetasoundExperimental/Source/MetasoundExperimentalRuntime/Private/MetasoundExampleNodeConfiguration.cpp
--- a/Engine/Plugins/Experimental/MetasoundExperimental/Source/MetasoundExperimentalRuntime/Private/MetasoundExampleNodeConfiguration.cpp
+++ b/Engine/Plugins/Experimental/MetasoundExperimental/Source/MetasoundExperimentalRuntime/Private/MetasoundExampleNodeConfiguration.cpp
@@ -23,11 +23,7 @@ namespace Metasound::Experimental
 	namespace ExampleNodeConfigurationPrivate
 	{
 		const FLazyName InputBaseName{"In"};
-#if WITH_EDITOR
-		const FText InputTooltip = LOCTEXT("In_ToolTip", "A trigger");
-#else
-		const FText InputTooltip = FText::GetEmpty();
-#endif
+		const FText InputTooltip = METASOUND_LOCTEXT("In_ToolTip", "A trigger");
 
 		FName MakeInputVertexName(int32 InIndex)
 		{
@@ -39,20 +35,12 @@ namespace Metasound::Experimental
 		FInputDataVertex MakeInputDataVertex(int32 InIndex)
 		{
 			FName InputName = MakeInputVertexName(InIndex);
-#if WITH_EDITOR
-			const FText InputDisplayName = FText::Format(LOCTEXT("In_DisplayName", "In {0}"), InIndex);
-#else
-			const FText InputDisplayName = FText::GetEmpty();
-#endif
+			const FText InputDisplayName = METASOUND_LOCTEXT_FORMAT("In_DisplayName", "In {0}", InIndex);
 			return TInputDataVertex<FTrigger>{InputName, FDataVertexMetadata{InputTooltip, InputDisplayName}};
 		}
 
 		FLazyName OutputBaseName{"Out"};
-#if WITH_EDITOR
-		const FText OutputTooltip = LOCTEXT("Out_ToolTip", "A trigger");
-#else
-		const FText OutputTooltip = FText::GetEmpty();
-#endif
+		const FText OutputTooltip = METASOUND_LOCTEXT("Out_ToolTip", "A trigger");
 
 		FName MakeOutputVertexName(int32 OutOutdex)
 		{
@@ -64,11 +52,7 @@ namespace Metasound::Experimental
 		FOutputDataVertex MakeOutputDataVertex(int32 OutOutdex)
 		{
 			FName OutputName = MakeOutputVertexName(OutOutdex);
-#if WITH_EDITOR
-			const FText OutputDisplayName = FText::Format(LOCTEXT("Out_DisplayName", "Out {0}"), OutOutdex);
-#else
-			const FText OutputDisplayName = FText::GetEmpty();
-#endif
+			const FText OutputDisplayName = METASOUND_LOCTEXT_FORMAT("Out_DisplayName", "Out {0}", OutOutdex);
 			return TOutputDataVertex<FTrigger>{OutputName, FDataVertexMetadata{OutputTooltip, OutputDisplayName}};
 		}
 
@@ -199,10 +183,10 @@ namespace Metasound::Experimental
 				FNodeClassName{ "Experimental", "ConfigurableOperator", "" },
 				1, // Major version
 				0, // Minor version
-				LOCTEXT("ExampleConfigurableNodeName", "A Configurable Node"),	
-				LOCTEXT("ExampleConfigurableNodeDescription", "A Node which shows how to make a configurable node for yourself."),	
+				METASOUND_LOCTEXT("ExampleConfigurableNodeName", "A Configurable Node"),	
+				METASOUND_LOCTEXT("ExampleConfigurableNodeDescription", "A Node which shows how to make a configurable node for yourself."),	
 				TEXT("UE"), // Author
-				LOCTEXT("ExampleConfigurablePromptIfMissing", "Enable the MetaSoundExperimental Plugin"), // Prompt if missing
+				METASOUND_LOCTEXT("ExampleConfigurablePromptIfMissing", "Enable the MetaSoundExperimental Plugin"), // Prompt if missing
 				ExampleNodeConfigurationPrivate::GetVertexInterface(1 /* default num inputs */, 1 /* default num outputs */),
 				{}
 			};
@@ -329,10 +313,10 @@ namespace Metasound::Experimental
 				FNodeClassName{ "Experimental", "WidgetConfigurableOperator", "" },
 				1, // Major version
 				0, // Minor version
-				LOCTEXT("WidgetExampleConfigurableNodeName", "A Widget Configurable Node"),
-				LOCTEXT("WidgetExampleConfigurableNodeDescription", "A Node which shows how to make a configurable node with a custom details customization for yourself."),
+				METASOUND_LOCTEXT("WidgetExampleConfigurableNodeName", "A Widget Configurable Node"),
+				METASOUND_LOCTEXT("WidgetExampleConfigurableNodeDescription", "A Node which shows how to make a configurable node with a custom details customization for yourself."),
 				TEXT("UE"), // Author
-				LOCTEXT("WidgetExampleConfigurablePromptIfMissing", "Enable the MetaSoundExperimental Plugin"), // Prompt if missing
+				METASOUND_LOCTEXT("WidgetExampleConfigurablePromptIfMissing", "Enable the MetaSoundExperimental Plugin"), // Prompt if missing
 				GetVertexInterface(),
 				{}
 			};
