[MetaSound Experimental] Fade Node unit test fix

#rnx
#tests Frontend MetaSound Tests

diff --git a/Engine/Plugins/Experimental/MetasoundExperimental/Source/MetasoundExperimentalEngineRuntime/Private/MetasoundFadeNode.cpp b/Engine/Plugins/Experimental/MetasoundExperimental/Source/MetasoundExperimentalEngineRuntime/Private/MetasoundFadeNode.cpp
--- a/Engine/Plugins/Experimental/MetasoundExperimental/Source/MetasoundExperimentalEngineRuntime/Private/MetasoundFadeNode.cpp
+++ b/Engine/Plugins/Experimental/MetasoundExperimental/Source/MetasoundExperimentalEngineRuntime/Private/MetasoundFadeNode.cpp
@@ -611,16 +611,22 @@ namespace Metasound
 
 			if (OutputType == EMetaSoundFadeOutputType::AudioBufferType)
 			{
-				SampleRate = InParams.OperatorSettings.GetSampleRate();
+				SampleRate = FMath::Max(1.f, InParams.OperatorSettings.GetSampleRate());
 			}
 			else if (OutputType == EMetaSoundFadeOutputType::FloatType)
 			{
-				SampleRate = InParams.OperatorSettings.GetActualBlockRate();
+				SampleRate = FMath::Max(1.f, InParams.OperatorSettings.GetActualBlockRate());
 			}
 
 			ActualFramesPerBlock = SampleRate / InParams.OperatorSettings.GetActualBlockRate();
 
-			InternalReset();
+			// Reset outputs
+			OnTrigger->Reset();
+			OnDone->Reset();
+			OnNearlyDone->Reset();
+
+			*FloatFadeValue = 0.f;
+			AudioFadeValue->Zero();
 		}
 
 		virtual ~TFadeNodeOperator() override = default;
@@ -676,7 +682,13 @@ namespace Metasound
 
 			ActualFramesPerBlock = SampleRate / InParams.OperatorSettings.GetActualBlockRate();
 
-			InternalReset();
+			// Reset outputs
+			OnTrigger->Reset();
+			OnDone->Reset();
+			OnNearlyDone->Reset();
+
+			*FloatFadeValue = 0.f;
+			AudioFadeValue->Zero();
 		}
 
 		void Execute()
@@ -739,10 +751,13 @@ namespace Metasound
 					check(StartFrame >= 0);
 
 					// Determine maximum slice size
-					const int32 MaxSliceSize = FMath::Clamp(EndFrame - StartFrame, 1, ActualFramesPerBlock);
+					const int32 MaxSliceSize = FMath::Clamp(EndFrame - StartFrame, 0, ActualFramesPerBlock);
+
+					// Determine best slice start frame
+					const int32 SliceStartFrame = FMath::Clamp(StartFrame, 0, ActualFramesPerBlock - 1);
 
 					// Get on pre trigger frame view
-					TArrayView<float> FadeValueViewPreTrigger = FadeValueView.Slice(StartFrame, MaxSliceSize);
+					TArrayView<float> FadeValueViewPreTrigger = FadeValueView.Slice(SliceStartFrame, MaxSliceSize);
 
 					// If fade is active, check for doneness
 					if (Fade.IsActive())
@@ -771,11 +786,11 @@ namespace Metasound
 					{
 						// Setup on pre trigger Fade Out frame slice
 						TArrayView<float> FadeOutBufferView(FadeOutBuffer);
-						TArrayView<float> CrossfadeBufferViewSlice = FadeOutBufferView.Slice(StartFrame, MaxSliceSize);
+						TArrayView<float> CrossfadeBufferViewSlice = FadeOutBufferView.Slice(SliceStartFrame, MaxSliceSize);
 
 						// Setup on pre trigger Fade In frame slice
 						TArrayView<float> FadeInBufferView(FadeInBuffer);
-						TArrayView<float> FadeInBufferViewSlice = FadeInBufferView.Slice(StartFrame, MaxSliceSize);
+						TArrayView<float> FadeInBufferViewSlice = FadeInBufferView.Slice(SliceStartFrame, MaxSliceSize);
 
 						// Apply fade in to value buffer
 						Audio::ArrayMultiplyInPlace(FadeInBufferViewSlice, FadeValueViewPreTrigger);
@@ -793,16 +808,19 @@ namespace Metasound
 					OnTrigger->TriggerFrame(StartFrame);
 
 					// Get maximum slice size for this on trigger block
-					const int32 MaxSliceSize = FMath::Clamp(EndFrame - StartFrame, 1, ActualFramesPerBlock);
+					const int32 MaxSliceSize = FMath::Clamp(EndFrame - StartFrame, 0, ActualFramesPerBlock);
+
+					// Determine best slice start frame
+					const int32 SliceStartFrame = FMath::Clamp(StartFrame, 0, ActualFramesPerBlock - 1);
 
 					// Determine start sample if seeking into the middle of the fade playback
-					const int32 StartSample = FMath::Floor(StartTimeIn->GetSeconds() * static_cast<float>(SampleRate));
+					const int32 StartSample = FMath::Floor(StartTime.GetSeconds() * static_cast<float>(SampleRate));
 
 					// Pass in which sample to start fade playback on
 					Fade.Start(StartSample);
 
 					// Create an on trigger slice of the value buffer
-					TArrayView<float> FadeValueViewOnTrigger = FadeValueView.Slice(StartFrame, MaxSliceSize);
+					TArrayView<float> FadeValueViewOnTrigger = FadeValueView.Slice(SliceStartFrame, MaxSliceSize);
 
 					// Get OnDone Frames
 					const int32 OnDoneFrame = Fade.DoneWithinInFrameRange(MaxSliceSize);
@@ -828,11 +846,11 @@ namespace Metasound
 
 						// Setup on on trigger Fade Out frame slice
 						TArrayView<float> FadeOutBufferView(FadeOutBuffer);
-						TArrayView<float> CrossfadeBufferViewSlice = FadeOutBufferView.Slice(StartFrame, MaxSliceSize);
+						TArrayView<float> CrossfadeBufferViewSlice = FadeOutBufferView.Slice(SliceStartFrame, MaxSliceSize);
 						
 						// Setup on on trigger Fade In frame slice
 						TArrayView<float> FadeInBufferView(FadeInBuffer);
-						TArrayView<float> FadeInBufferViewSlice = FadeInBufferView.Slice(StartFrame, MaxSliceSize);
+						TArrayView<float> FadeInBufferViewSlice = FadeInBufferView.Slice(SliceStartFrame, MaxSliceSize);
 
 						// Apply fade in to value buffer
 						Audio::ArrayMultiplyInPlace(FadeInBufferViewSlice, FadeValueViewOnTrigger);
@@ -866,14 +884,8 @@ namespace Metasound
 		{
 			using namespace FadeNodePrivate;
 
-			// Update cached input values
-			Duration = *DurationIn;
-			StartValue = *StartValueIn;
-			EndValue = *EndValueIn;
-			StartTime = *StartTimeIn;
-			bResetValue = *ResetValueIn;
-			CurveFunction = *CurveFunctionIn;
-			Exponent = *ExponentIn;
+			// Cache and guard against bad input
+			CacheParams();
 
 			// calculate number of samples in fade
 			const float TotalFadeSamples = FMath::Max(0.f, Duration.GetSeconds()) * SampleRate;
@@ -882,29 +894,35 @@ namespace Metasound
 			FFadeData Data =
 			{
 				TotalFadeSamples,
-				*StartValueIn,
-				*EndValueIn,
-				*ExponentIn,
-				*CurveFunctionIn,
-				*ResetValueIn
+				StartValue,
+				EndValue,
+				Exponent,
+				CurveFunction,
+				bResetValue
 			};
 
 			// Initialize Fade with Fade Data
 			Fade.Init(Data);
 		}
 
-		void UpdateParams()
+		void CacheParams()
 		{
-			using namespace FadeNodePrivate;
-
-			// Cache input values
-			Duration = *DurationIn;
+			// Update cached input values
+			Duration = FTime::FromSeconds(FMath::Clamp(DurationIn->GetSeconds(), 0.f, 10000.f));
 			StartValue = *StartValueIn;
 			EndValue = *EndValueIn;
-			StartTime = *StartTimeIn;
+			StartTime = FTime::FromSeconds(FMath::Clamp(StartTimeIn->GetSeconds(), 0.f, Duration.GetSeconds()));
 			bResetValue = *ResetValueIn;
 			CurveFunction = *CurveFunctionIn;
-			Exponent = *ExponentIn;
+			Exponent = FMath::Max(*ExponentIn, 0.f);
+		}
+
+		void UpdateParams()
+		{
+			using namespace FadeNodePrivate;
+
+			// Cache and guard against bad input
+			CacheParams();
 
 			// Cacluate latest total number of fade samples
 			const float TotalFadeSamples = FMath::Max(0.f, Duration.GetSeconds()) * SampleRate;
@@ -913,11 +931,11 @@ namespace Metasound
 			FFadeData Data =
 			{
 				TotalFadeSamples,
-				*StartValueIn,
-				*EndValueIn,
-				*ExponentIn,
-				*CurveFunctionIn,
-				*ResetValueIn
+				StartValue,
+				EndValue,
+				Exponent,
+				CurveFunction,
+				bResetValue
 			};
 
 			// Update the fade data
