[Experimental] [Metasounds] Polymorphic CAT Update and Granulator node

#tests LowLevelTests

diff --git a/Engine/Plugins/Experimental/MetasoundExperimental/Source/AudioExperimentalRuntime/Private/ChannelAgnostic/ChannelAgnosticTranscoding.cpp b/Engine/Plugins/Experimental/MetasoundExperimental/Source/AudioExperimentalRuntime/Private/ChannelAgnostic/ChannelAgnosticTranscoding.cpp
--- a/Engine/Plugins/Experimental/MetasoundExperimental/Source/AudioExperimentalRuntime/Private/ChannelAgnostic/ChannelAgnosticTranscoding.cpp
+++ b/Engine/Plugins/Experimental/MetasoundExperimental/Source/AudioExperimentalRuntime/Private/ChannelAgnostic/ChannelAgnosticTranscoding.cpp
@@ -6,6 +6,7 @@
 #include "DSP/Ambisonics.h"
 #include "DSP/MultiMono.h"
 #include "DSP/SphericalHarmonicCalculator.h"
+#include "DSP/Vbap.h"
 #include "TypeFamily/ChannelTypeFamily.h"
 
 namespace Audio
@@ -58,6 +59,47 @@ namespace Audio
 			}
 		}
 	}
+
+
+
+	static const IDiscretePanner* ChoosePanner(const FDiscreteChannelTypeFamily& InFrom, const FDiscreteChannelTypeFamily& InTo)
+	{
+		// Return To panner.
+		// All using 2d vbap for now, so just return that.
+		return InTo.GetPanner();
+	}
+	
+	// Static speakers can generate a matrix of speaker gains and cache that to avoid running the panner everytime.
+	bool MakeGainMatrixFrom(const FDiscreteChannelTypeFamily& InSrc, const FDiscreteChannelTypeFamily& InDst, TArray<float>& OutMatrix)
+	{
+		if (const IDiscretePanner* Panner = ChoosePanner(InSrc, InDst))
+		{
+			const int32 DstChannels = InDst.NumChannels();
+			const int32 SrcChannels = InSrc.NumChannels();
+			OutMatrix.SetNum(SrcChannels*DstChannels);
+
+			const TArray<FDiscreteChannelTypeFamily::FSpeaker>& Order = InSrc.GetSpeakerOrder();
+			
+			// For each src channel, create row vector of gains for each dst channel.
+			TArray<float> SpeakerGains;
+			SpeakerGains.SetNum(DstChannels);
+			for (int32 i = 0; i < SrcChannels; ++i)
+			{
+				IDiscretePanner::FOutputParams Results;
+				IDiscretePanner::FInputParams Params { .AzimuthDegrees = Order[i].AzimuthDegrees, .ElevationDegrees = Order[i].ElevationDegrees };
+				if (Panner->ComputeGains(Params,Results) )
+				{
+					// Copy results into gain array.
+					for (int32 j = 0; j < Results.Results.Num(); ++j)
+					{
+						OutMatrix[i*DstChannels+j] = Results.Results[j].Gain;
+					}
+				}
+			}
+			return true;
+		}
+		return false;
+	}
 	
 	// Discrete to Discrete
 	FChannelTypeFamily::FTranscoder GetTranscoder(
@@ -101,7 +143,16 @@ namespace Audio
 
 			case EChannelTranscodeMethod::MixUpOrDown:
 			{
-				// Make a mix matrix and call a mix/up down.
+#if 1 
+				if (TArray<float> Gains; MakeGainMatrixFrom(InSrcType, InDstType, Gains))
+				{
+					return [MixGains = MoveTemp(Gains)](TArrayView<const float*> InSrc, TArrayView<float*> InDst, const int32 NumFrames)
+					{
+						MultiMonoMixUpOrDown(InSrc, InDst, NumFrames, MixGains);
+					};					
+				}
+#else
+					// Make a mix matrix and call a mix/up down.
 				if (TArray<float> Gains; Create2DChannelMap(
 					{
 						.NumInputChannels = InSrcType.NumChannels(),
@@ -120,6 +171,7 @@ namespace Audio
 						MultiMonoMixUpOrDown(InSrc, InDst, NumFrames, MixGains);
 					};
 				}
+#endif //					
 			}
 		}
 		// fail.
@@ -128,7 +180,7 @@ namespace Audio
 
 	// Discrete to Ambisonics
 	FChannelTypeFamily::FTranscoder GetTranscoder(const FDiscreteChannelTypeFamily& InFromType,
-		const FAmbisonicsChannelTypeFamily& InToType, const FChannelTypeFamily::FGetTranscoderParams&)
+		const FSoundfieldChannelTypeFamily& InToType, const FChannelTypeFamily::FGetTranscoderParams&)
 	{
 		const TArray<FDiscreteChannelTypeFamily::FSpeaker>& Speakers = InFromType.GetSpeakerOrder();
 		TArray<TArray<float>> AllChannelGainArrays;
@@ -170,7 +222,7 @@ namespace Audio
 	}
 	
 	// Ambisonics to Discrete
-	FChannelTypeFamily::FTranscoder GetTranscoder(const FAmbisonicsChannelTypeFamily& InFromType, const FDiscreteChannelTypeFamily& InToType, const FChannelTypeFamily::FGetTranscoderParams&)
+	FChannelTypeFamily::FTranscoder GetTranscoder(const FSoundfieldChannelTypeFamily& InFromType, const FDiscreteChannelTypeFamily& InToType, const FChannelTypeFamily::FGetTranscoderParams&)
 	{
 		TArray<TArray<float>> AllChannelGainArrays;
 		AllChannelGainArrays.SetNum(InToType.NumChannels());
@@ -190,7 +242,7 @@ namespace Audio
 			
 			//FVector2f Spherical = Order[i].Position.UnitCartesianToSpherical();
 			//FSphericalHarmonicCalculator::AdjustUESphericalCoordinatesForAmbisonics(Spherical);
-			check(GainArray.Num() == FAmbisonicsChannelTypeFamily::OrderToNumChannels(InFromType.GetAmbisonicsOrder()));
+			check(GainArray.Num() == FSoundfieldChannelTypeFamily::OrderToNumChannels(InFromType.GetAmbisonicsOrder()));
 			const float AzimuthRads = FMath::DegreesToRadians(Speaker.AzimuthDegrees);
 			const float ElevationRads = FMath::DegreesToRadians(Speaker.ElevationDegrees);
 			FSphericalHarmonicCalculator::ComputeSoundfieldChannelGains(InFromType.GetAmbisonicsOrder(), AzimuthRads, ElevationRads, GainArray);
@@ -209,8 +261,8 @@ namespace Audio
 	}
 
 	// Ambisonics to Ambisonics
-	FChannelTypeFamily::FTranscoder GetTranscoder(const FAmbisonicsChannelTypeFamily& InFromType,
-		const FAmbisonicsChannelTypeFamily& InToType, const FChannelTypeFamily::FGetTranscoderParams& InParams)
+	FChannelTypeFamily::FTranscoder GetTranscoder(const FSoundfieldChannelTypeFamily& InFromType,
+		const FSoundfieldChannelTypeFamily& InToType, const FChannelTypeFamily::FGetTranscoderParams& InParams)
 	{
 		// Ambisonics to Ambisonics
 		if (&InToType == &InFromType)
diff --git a/Engine/Plugins/Experimental/MetasoundExperimental/Source/AudioExperimentalRuntime/Private/TypeFamily/ChannelTypeFamily.cpp b/Engine/Plugins/Experimental/MetasoundExperimental/Source/AudioExperimentalRuntime/Private/TypeFamily/ChannelTypeFamily.cpp
--- a/Engine/Plugins/Experimental/MetasoundExperimental/Source/AudioExperimentalRuntime/Private/TypeFamily/ChannelTypeFamily.cpp
+++ b/Engine/Plugins/Experimental/MetasoundExperimental/Source/AudioExperimentalRuntime/Private/TypeFamily/ChannelTypeFamily.cpp
@@ -4,9 +4,12 @@
 
 #include "Algo/Transform.h"
 #include "CoreGlobals.h"
+
+#include "HAL/FileManager.h"
 #include "ChannelAgnostic/ChannelAgnosticTranscoding.h"
 #include "Containers/Map.h"
 #include "DSP/SphericalHarmonicCalculator.h"
+#include "DSP/Vbap.h"
 
 namespace Audio
 {
@@ -23,17 +26,17 @@ namespace Audio
 	
 	class FChannelRegistryImpl final : public IChannelTypeRegistry 
 	{
-		TMap<FName, const FChannelTypeFamily*> Types;
+		TMap<FName, TSharedRef<const FChannelTypeFamily>> Types;
 	public:
 		UE_NONCOPYABLE(FChannelRegistryImpl);
 		FChannelRegistryImpl() = default;
 		virtual ~FChannelRegistryImpl() override = default;
 	protected:
-		virtual bool RegisterType(const FName& InUniqueName, TRetainedRef<FTypeFamily> InType) override
+		virtual bool RegisterType(const FName& InUniqueName, TUniquePtr<FTypeFamily> InType) override
 		{
 			static const FName ChannelTypeName = TEXT("Cat");
 			const FTypeFamily* Type = FindTypeInternal(ChannelTypeName);
-			if (Type && !InType.Get().IsA(Type))
+			if (Type && !InType->IsA(Type))
 			{
 				// Not a cat. (maybe a dog?)
 				return false;	
@@ -46,36 +49,55 @@ namespace Audio
 			}
 			
 			// Safe to cast it and add.
-			const FChannelTypeFamily* ChannelType = static_cast<const FChannelTypeFamily*>(&InType.Get());
+			TSharedRef<FChannelTypeFamily> Ref(static_cast<FChannelTypeFamily*>(InType.Release()));
 			
 			// Add to master list of registered types.
 			check(Types.Find(InUniqueName) == nullptr);
-			Types.Add(InUniqueName,ChannelType);
-	
+			Types.Emplace(InUniqueName,MoveTemp(Ref));
+			
 			return true;
 		}	
 		virtual const FTypeFamily* FindTypeInternal(const FName InUniqueName) const override
 		{
-			if (const FChannelTypeFamily* const * Found = Types.Find(InUniqueName))
+			if (const auto Found = Types.Find(InUniqueName))
 			{
-				return *Found;
+				return &(Found->Get());
 			}
 			return nullptr;
 		}
 
-		virtual TArray<const FChannelTypeFamily*> GetAllChannelFormats() const override
+		virtual TArray<TSharedRef<const FChannelTypeFamily>> GetAllChannelFormats() const override
 		{
-			TArray<const FChannelTypeFamily*> AllFormats;
+			TArray<TSharedRef<const FChannelTypeFamily>> AllFormats;
 			Types.GenerateValueArray(AllFormats);
 			return AllFormats;
 		};
 	};
 
-	int32 FAmbisonicsChannelTypeFamily::OrderToNumChannels(const int32 InOrder)
+	int32 FSoundfieldChannelTypeFamily::OrderToNumChannels(const int32 InOrder)
 	{
 		return FSphericalHarmonicCalculator::OrderToNumChannels(InOrder);
 	}
 
+	TOptional<FChannelTypeFamily::FChannelName> FSoundfieldChannelTypeFamily::GetChannelName(const int32 InChannelIndex) const
+	{
+		if (InChannelIndex < 0)
+		{
+			return {};
+		}
+
+		int32 OrderNum=0;
+		int32 DegreeNum=0;
+		FSphericalHarmonicCalculator::IndexToOrderAndDegree(InChannelIndex, OrderNum, DegreeNum);
+		const FName ChannelName = *FString::Printf(TEXT("ACN %d"), InChannelIndex); // *ChannelName.ToString(),	 (<-- Indexed' FName appears to fail for Node config UX, so use Sprintf'd one).ChannelName(TEXT("ACN"),InChannelIndex; 
+		return FChannelName
+		{
+			.Name = ChannelName,  
+			.FriendlyName = FString::Printf(TEXT("ACN=%d, Order=%d, Degree=%d"), InChannelIndex, Order, DegreeNum)
+		};
+	}
+	
+
 	IChannelTypeRegistry& GetChannelRegistry()
 	{
 		static FChannelRegistryImpl Registry;
@@ -116,14 +138,17 @@ namespace Audio
    FDiscreteChannelTypeFamily::FDiscreteChannelTypeFamily(const FName& InUniqueName, FChannelTypeFamily* InParentType, const FString& InFriendlyName, const TArray<FSpeaker>& InOrder, const bool bIsParentsDefault, const bool bIsAbstract)
 		: Super(InUniqueName, GetFamilyTypeName(), InOrder.Num(), InParentType, InFriendlyName, bIsParentsDefault, bIsAbstract)
 		, Order(InOrder)
+		, Panner(MakePanner(InOrder))
 	{
 		checkf(InParentType, TEXT("Type=%s, Has a Null Parent"), *GetName().ToString());
-		auto GetParentNameSafe = [](FChannelTypeFamily* i) -> FString { return i ? i->GetName().ToString() : FString(); };
-		UE_LOG(LogTemp, Display, TEXT("Unique=%s\tNumChannels=%d\tParent=%s\tFriendlyName=%s\tDefault=%s\tOrder=[%s]\tAbstract=%s"),
-			*InUniqueName.ToString(),  InOrder.Num(), *GetParentNameSafe(InParentType), *InFriendlyName, ToCStr(LexToString(bIsParentsDefault)), *ChannelTypeFamilyPrviate::MakePrettyString(Order), ToCStr(LexToString(bIsAbstract)));
 	}
 
-	TOptional<FChannelTypeFamily::FChannelName> FDiscreteChannelTypeFamily::GetChannelName(const int32 InChannelIndex) const
+    FDiscreteChannelTypeFamily::~FDiscreteChannelTypeFamily()
+	{
+		// Needed for TUnique deleter, as panner is forward declared in header.
+	}
+
+   TOptional<FChannelTypeFamily::FChannelName> FDiscreteChannelTypeFamily::GetChannelName(const int32 InChannelIndex) const
 	{
 		check(InChannelIndex >= 0 && InChannelIndex < NumChannels());
 		if (Order.IsValidIndex(InChannelIndex))
@@ -141,245 +166,60 @@ namespace Audio
 		return {};
 	}
 
-	void RegisterChannelLayouts(IChannelTypeRegistry& Registry)
+	TUniquePtr<IDiscretePanner> FDiscreteChannelTypeFamily::MakePanner(const TArray<FSpeaker>& InOrder)
 	{
-		/**
-		 * Standard layouts.
-		 * This should be defined in .ini ultimately, allowing new, custom formats to be added simply.
-		 **/
-		
-		// Register root type.
-		static FChannelTypeFamily BaseCat(TEXT("Cat"), TEXT("Cat"),  0, nullptr, TEXT("Base Cat"), false, true);
-		Registry.RegisterType(BaseCat.GetName(), BaseCat);
-		
-#define REGISTER_CAT(NAME, PARENT_NAME, FRIENDLY_NAME, LAYOUT, DEFAULT,ABSTRACT)\
-static FDiscreteChannelTypeFamily UE_JOIN(__Type,NAME)(UE_STRINGIZE(NAME), Registry.FindChannel(PARENT_NAME), FRIENDLY_NAME, LAYOUT, DEFAULT, ABSTRACT);\
-Registry.RegisterType(TEXT(UE_STRINGIZE(NAME)),UE_JOIN(__Type,NAME))
-
-		using enum ESpeakerShortNames;
-		using FS = FDiscreteChannelTypeFamily::FSpeaker;
-
-		// Top level abstraction. "Discrete".
-		REGISTER_CAT(Discrete,			TEXT("Cat"),		TEXT("Discrete"),		{},										false,					true);		
-
-		// Mono
-		REGISTER_CAT(Mono, TEXT("Discrete"), TEXT("Mono"), {}, false, true);
-		REGISTER_CAT(Mono1Dot0, TEXT("Mono"), TEXT("Mono (1.0)"),
-		             (TArray<FS>
-			             {
-			             { FC, 0.f, 0.f }, // Az:   0°, El:  0°
-			             }), true, false);
-
-		REGISTER_CAT(Mono1Dot1, TEXT("Mono"), TEXT("Mono (1.1)"), (TArray<FS>{
-			             { FC, 0.f, 0.f }, // Az:   0°, El:  0°
-			             { LFE, 0.f, 0.f } // LFE (omni)
-			             }), false, false);
+		// Just using 2D VBap for now.
+		TArray<IDiscretePanner::FSpeaker> Speakers;
+		for (int32 Channel = 0; Channel < InOrder.Num(); Channel++)
+		{
+			const FSpeaker& Speaker = InOrder[Channel];
 
-// Stereo
-REGISTER_CAT(Stereo, TEXT("Discrete"), TEXT("Stereo"), {}, false, true);
-REGISTER_CAT(Stereo2Dot0, TEXT("Stereo"), TEXT("Stereo (2.0)"),
-	(TArray<FS>{
-		{ FL, -30.f,   0.f },  // Az: -30°, El:  0°
-		{ FR,  30.f,   0.f }   // Az: +30°, El:  0°
-		}), true, false);
+			// Add all speakers that are considered 'spatialized'
+			// This excludes channels that not, like LFE and possibly (Centre channel for dialog). 
+			if (Speaker.bIsSpatialized)
+			{
+				IDiscretePanner::FSpeaker PannerSpeaker
+				{
+					.AzimuthDegrees = Speaker.AzimuthDegrees,
+					.ElevationDegrees = Speaker.ElevationDegrees,
+					.ChannelID = static_cast<int32>(Speaker.Speaker),
+					.ChannelIndex = Channel
+				};
+				Speakers.Add(PannerSpeaker);
+			}
+		}
 		
-REGISTER_CAT(Stereo2Dot1, TEXT("Stereo"), TEXT("Stereo (2.1)"),
-	(TArray<FS>{
-		{ FL, -30.f, 0.f },		// Az: -30°, El:  0°
-		{ FR, 30.f, 0.f },		// Az: +30°, El:  0°
-		{ LFE, 0.f, 0.f }		// LFE (omni)
-	}),
-	false, false);
-REGISTER_CAT(Stereo3Dot0, TEXT("Stereo"), TEXT("Stereo (3.0)"),
-	(TArray<FS>{
-		{ FL, -30.f,   0.f },  // Az: -30°, El:  0°
-		{ FC,   0.f,   0.f },  // Az:   0°, El:  0°
-		{ FR,  30.f,   0.f }   // Az: +30°, El:  0°
-	}),
-	false, false);
-REGISTER_CAT(Stereo3DOt1, TEXT("Stereo"), TEXT("Stereo (3.1)"),
-	(TArray<FS>{
-		{ FL, -30.f,   0.f },  // Az: -30°, El:  0°
-		{ FC,   0.f,   0.f },  // Az:   0°, El:  0°
-		{ FR,  30.f,   0.f },  // Az: +30°, El:  0°
-		{ LFE,  0.f,   0.f }   // LFE (omni)
-	}),
-	false, false);
-
-// Quad
-REGISTER_CAT(Quad, TEXT("Discrete"), TEXT("Quad"), {}, false, true);
-REGISTER_CAT(Quad4Dot0Back, TEXT("Quad"), TEXT("Quad Back Speakers (4.0)"),
-	(TArray<FS>{
-		{ FL, -30.f,   0.f },  // Front Left
-		{ FR,  30.f,   0.f },  // Front Right
-		{ BL, -150.f,  0.f },  // Rear  Left
-		{ BR, 150.f,   0.f }   // Rear  Right
-		}),
-	true, false);
-REGISTER_CAT(Quad4Dot0Side, TEXT("Quad"), TEXT("Quad Side Speakers (4.0)"),
-	(TArray<FS>{
-		{ FL, -30.f,   0.f },  // Front Left
-		{ FR,  30.f,   0.f },  // Front Right
-		{ SL, -90.f,   0.f },  // Side  Left
-		{ SR,  90.f,   0.f }   // Side  Right
-		}),
-		false, false);
-REGISTER_CAT(Quad4Dot1, TEXT("Quad"), TEXT("Quad Back Centre LFE (4.1)"),
-	(TArray<FS>{
-		{ FL, -30.f,   0.f },  // Front Left
-		{ FR,  30.f,   0.f },  // Front Right
-		{ BL, -150.f,  0.f },  // Rear  Left
-		{ BR, 150.f,   0.f },  // Rear  Right
-		{ LFE, 0.f,    0.f }   // LFE (omni)
-		}),
-	false, false);
-
-// Surround
-REGISTER_CAT(Surround, TEXT("Discrete"), TEXT("Surround"), {}, false, true);
-REGISTER_CAT(Surround5, TEXT("Surround"), TEXT("Surround (5.X)"), {}, true, true);
-REGISTER_CAT(Surround5Dot0, TEXT("Surround5"), TEXT("Surround (5.0)"),
-    (TArray<FS>{
-        { FL,  -30,  0 },
-        { FR,   30,  0 },
-        { SL,  -90,  0 },
-        { SR,   90,  0 },
-        { FC,    0,  0 },
-        { BL, -135,  0 },
-        { BR,  135,  0 },
-    }),
-    false, false);
-
-REGISTER_CAT(Surround5_1, TEXT("Surround5"), TEXT("Surround (5.1)"),
-    (TArray<FS>{
-        { FL,  -30,  0 },
-        { FR,   30,  0 },
-        { SL,  -90,  0 },
-        { SR,   90,  0 },
-        { FC,    0,  0 },
-        { BL, -135,  0 },
-        { BR,  135,  0 },
-        { LFE,   0,  0 },
-    }),
-    true, false);
-
-REGISTER_CAT(Surround7, TEXT("Surround"), TEXT("Surround (7.X)"),
-    {}, false, true);
-
-REGISTER_CAT(Surround7Dot0, TEXT("Surround7"), TEXT("Surround (7.0)"),
-    (TArray<FS>{
-        { FL,  -30,  0 },
-        { FR,   30,  0 },
-        { SL,  -90,  0 },
-        { SR,   90,  0 },
-        { FC,    0,  0 },
-        { BL, -135,  0 },
-        { BR,  135,  0 },
-    }),
-    false, false);
-
-REGISTER_CAT(Surround7Dot1, TEXT("Surround7"), TEXT("Surround (7.1)"),
-    (TArray<FS>{
-        { FL,  -30,  0 },
-        { FR,   30,  0 },
-        { SL,  -90,  0 },
-        { SR,   90,  0 },
-        { FC,    0,  0 },
-        { LFE,   0,  0 },
-        { BL, -135,  0 },
-        { BR,  135,  0 },
-    }),
-    true, false);
-
-// Atmos beds & heights
-REGISTER_CAT(Atmos, TEXT("Surround7"), TEXT("Dolby Atmos Bed"),
-    {}, false, true);
-
-REGISTER_CAT(Atmos7Dot0Dot2, TEXT("Atmos"), TEXT("Dolby Atmos (7.0.2)"),
-    (TArray<FS>{
-        { FL,  -30.f,   0.f },
-        { FR,   30.f,   0.f },
-        { SL,  -90.f,   0.f },
-        { SR,   90.f,   0.f },
-        { FC,    0.f,   0.f },
-        { BL, -135.f,   0.f },
-        { BR,  135.f,   0.f },
-        // heights
-        { TFL, -35.3f,  30.0f },
-        { TFR,  35.3f,  30.0f },
-    }),
-    false, false);
-
-REGISTER_CAT(Atmos7Dot0Dot4, TEXT("Atmos"), TEXT("Dolby Atmos (7.0.4)"),
-    (TArray<FS>{
-        { FL,  -30.f,   0.f },
-        { FR,   30.f,   0.f },
-        { SL,  -90.f,   0.f },
-        { SR,   90.f,   0.f },
-        { FC,    0.f,   0.f },
-        { BL, -135.f,   0.f },
-        { BR,  135.f,   0.f },
-        // heights
-        { TFL, -35.3f,  30.0f },
-        { TFR,  35.3f,  30.0f },
-        { TBL, -135.0f, 45.0f },
-        { TBR,  135.0f, 45.0f },
-    }),
-    false, false);
-
-REGISTER_CAT(Atmos7Dot1Dot2, TEXT("Atmos"), TEXT("Dolby Atmos (7.1.2)"),
-    (TArray<FS>{
-        { FL,  -30.f,   0.f },
-        { FR,   30.f,   0.f },
-        { SL,  -90.f,   0.f },
-        { SR,   90.f,   0.f },
-        { FC,    0.f,   0.f },
-        { LFE,   0.f,   0.f },
-        { BL, -135.f,   0.f },
-        { BR,  135.f,   0.f },
-        // heights
-        { TFL, -35.3f,  30.0f },
-        { TFR,  35.3f,  30.0f },
-    }),
-    false, false);
-
-REGISTER_CAT(Atmos7Dot1Dot4, TEXT("Atmos"), TEXT("Dolby Atmos (7.1.4)"),
-    (TArray<FS>{
-        { FL,  -30.f,    0.f },
-        { FR,   30.f,    0.f },
-        { SL,  -90.f,    0.f },
-        { SR,   90.f,    0.f },
-        { FC,    0.f,    0.f },
-        { LFE,   0.f,    0.f },
-        { BL, -135.f,    0.f },
-        { BR,  135.f,    0.f },
-        // heights
-        { TFL, -35.3f,   30.0f },
-        { TFR,  35.3f,   30.0f },
-        { TBL, -135.0f,  45.0f },
-        { TBR,  135.0f,  45.0f },
-    }),
-    true, false);
-
-					
-#define REGISTER_AMBISONICS(NAME, PARENT_NAME, FRIENDLY_NAME, ORDER, DEFAULT,ABSTRACT)\
-static FAmbisonicsChannelTypeFamily UE_JOIN(__Type,NAME)(UE_STRINGIZE(NAME), ORDER, Registry.FindChannel(PARENT_NAME), FRIENDLY_NAME, DEFAULT, ABSTRACT);\
-ensure(Registry.RegisterType(UE_STRINGIZE(NAME),UE_JOIN(__Type,NAME)))
-
-		// Ambisonics 
-		REGISTER_AMBISONICS(Ambisonics,					TEXT("Cat"), 	   TEXT("Ambisonics"), 				0, false, true);
-		REGISTER_AMBISONICS(Ambisonics1stOrder,			TEXT("Ambisonics"), TEXT("1st Order Ambisonics"), 		1, false, false); 
-		REGISTER_AMBISONICS(Ambisonics2ndOrder,			TEXT("Ambisonics"), TEXT("2nd Order Ambisonics"), 		2, false, false); 
-		REGISTER_AMBISONICS(Ambisonics3rdOrder,			TEXT("Ambisonics"), TEXT("3rd Order Ambisonics"), 		3, false, false);
-		REGISTER_AMBISONICS(Ambisonics4thOrder,			TEXT("Ambisonics"), TEXT("4th Order Ambisonics"),		4, false, false); 
-		REGISTER_AMBISONICS(Ambisonics5thOrder,			TEXT("Ambisonics"), TEXT("5th Order Ambisonics"),		5, false, false);  
-		REGISTER_AMBISONICS(Ambisonics6thOrder,			TEXT("Ambisonics"), TEXT("6th Order Ambisonics"),		6, false, false);  
-		REGISTER_AMBISONICS(Ambisonics7thOrder,			TEXT("Ambisonics"), TEXT("7th Order Ambisonics"),		7, false, false);  
+		if (Speakers.Num() > 0)
+		{
+			return MakeUnique<FVBap2D>(Speakers);
+		}
+		return {};
+	}
 
-		// Channel packs.
-		// FPackedChannelTypeFamily MonoPack( TEXT("PackedChannel")
-		
-#undef REGISTER_CAT
-#undef REGISTER_AMBISONICS
+	void RegisterChannelLayouts(IChannelTypeRegistry& Registry)
+	{
+		// Register root type.
+		Registry.RegisterType(TEXT("Cat"), MakeUnique<FChannelTypeFamily>(TEXT("Cat"), TEXT("Cat"),  0, nullptr, TEXT("Base Cat"), false, true));
+		FChannelTypeFamily* BaseCat = Registry.FindChannel(TEXT("Cat"));
+		check(BaseCat);
+
+		Registry.RegisterType(
+			TEXT("Cat:Discrete"),
+			MakeUnique<FDiscreteChannelTypeFamily>(
+						TEXT("Cat:Discrete"),
+						BaseCat,
+						TEXT("Discrete Channel Format"),
+						TArray<FDiscreteChannelTypeFamily::FSpeaker>(),
+						false,
+						true)
+			);
+		Registry.RegisterType(
+			TEXT("Cat:Soundfield"),
+			MakeUnique<FSoundfieldChannelTypeFamily>(
+				TEXT("Cat:Soundfield"),
+				0, BaseCat, TEXT("Soundfield Format"),
+				false,
+				true));
 	}
 }
 
diff --git a/Engine/Plugins/Experimental/MetasoundExperimental/Source/AudioExperimentalRuntime/Public/ChannelAgnostic/ChannelAgnosticTranscoding.h b/Engine/Plugins/Experimental/MetasoundExperimental/Source/AudioExperimentalRuntime/Public/ChannelAgnostic/ChannelAgnosticTranscoding.h
--- a/Engine/Plugins/Experimental/MetasoundExperimental/Source/AudioExperimentalRuntime/Public/ChannelAgnostic/ChannelAgnosticTranscoding.h
+++ b/Engine/Plugins/Experimental/MetasoundExperimental/Source/AudioExperimentalRuntime/Public/ChannelAgnostic/ChannelAgnosticTranscoding.h
@@ -8,9 +8,9 @@ namespace Audio
 {
 	// All transcoder (from -> to) combinations.
 	AUDIOEXPERIMENTALRUNTIME_API FChannelTypeFamily::FTranscoder GetTranscoder(const FDiscreteChannelTypeFamily& InFrom, const FDiscreteChannelTypeFamily& InTo, const FChannelTypeFamily::FGetTranscoderParams&);
-	AUDIOEXPERIMENTALRUNTIME_API FChannelTypeFamily::FTranscoder GetTranscoder(const FDiscreteChannelTypeFamily& InFrom, const FAmbisonicsChannelTypeFamily& InTo, const FChannelTypeFamily::FGetTranscoderParams&);
-	AUDIOEXPERIMENTALRUNTIME_API FChannelTypeFamily::FTranscoder GetTranscoder(const FAmbisonicsChannelTypeFamily& InFrom, const FDiscreteChannelTypeFamily& InTo, const FChannelTypeFamily::FGetTranscoderParams&);
-	AUDIOEXPERIMENTALRUNTIME_API FChannelTypeFamily::FTranscoder GetTranscoder(const FAmbisonicsChannelTypeFamily& InFrom, const FAmbisonicsChannelTypeFamily& InTo, const FChannelTypeFamily::FGetTranscoderParams&);
+	AUDIOEXPERIMENTALRUNTIME_API FChannelTypeFamily::FTranscoder GetTranscoder(const FDiscreteChannelTypeFamily& InFrom, const FSoundfieldChannelTypeFamily& InTo, const FChannelTypeFamily::FGetTranscoderParams&);
+	AUDIOEXPERIMENTALRUNTIME_API FChannelTypeFamily::FTranscoder GetTranscoder(const FSoundfieldChannelTypeFamily& InFrom, const FDiscreteChannelTypeFamily& InTo, const FChannelTypeFamily::FGetTranscoderParams&);
+	AUDIOEXPERIMENTALRUNTIME_API FChannelTypeFamily::FTranscoder GetTranscoder(const FSoundfieldChannelTypeFamily& InFrom, const FSoundfieldChannelTypeFamily& InTo, const FChannelTypeFamily::FGetTranscoderParams&);
 		
 	template<typename TFromType>
 	struct TFromVisitor final : IChannelTypeVisitor
@@ -27,7 +27,7 @@ namespace Audio
 		{
 			Transcoder = GetTranscoder(FromType, InTo, Params);	
 		}
-		virtual void Visit(const FAmbisonicsChannelTypeFamily& InTo) override
+		virtual void Visit(const FSoundfieldChannelTypeFamily& InTo) override
 		{
 			Transcoder = GetTranscoder(FromType, InTo, Params);
 		}
@@ -41,7 +41,7 @@ namespace Audio
 			TFromVisitor Visitor(InFrom, Params, Result);
 			Params.ToType.Accept(Visitor);
 		}
-		virtual void Visit(const FAmbisonicsChannelTypeFamily& InFrom) override
+		virtual void Visit(const FSoundfieldChannelTypeFamily& InFrom) override
 		{
 			TFromVisitor Visitor(InFrom, Params, Result);
 			Params.ToType.Accept(Visitor);
diff --git a/Engine/Plugins/Experimental/MetasoundExperimental/Source/AudioExperimentalRuntime/Public/TypeFamily/ChannelTypeFamily.h b/Engine/Plugins/Experimental/MetasoundExperimental/Source/AudioExperimentalRuntime/Public/TypeFamily/ChannelTypeFamily.h
--- a/Engine/Plugins/Experimental/MetasoundExperimental/Source/AudioExperimentalRuntime/Public/TypeFamily/ChannelTypeFamily.h
+++ b/Engine/Plugins/Experimental/MetasoundExperimental/Source/AudioExperimentalRuntime/Public/TypeFamily/ChannelTypeFamily.h
@@ -9,6 +9,11 @@
 #include "Misc/Optional.h"
 
 
+namespace Audio
+{
+	struct IDiscretePanner;
+}
+
 /**
  * A Comprehensive Short form Speaker enumeration
  */
@@ -31,6 +36,7 @@ enum class ESpeakerShortNames : uint32
 	TBR = EAudioMixerChannel::TopBackRight,			// Top Back Right
 
 	NumChannels = EAudioMixerChannel::ChannelTypeCount,
+	DefaultChannel = EAudioMixerChannel::DefaultChannel,
 };
 
 // Pasted from UHT generated.h, but we can't use because of AudioUnitTests.
@@ -51,7 +57,7 @@ enum class ESpeakerShortNames : uint32
 	op(ESpeakerShortNames::TBL) \
 	op(ESpeakerShortNames::TBR)
 
-const TCHAR* LexToString(const ESpeakerShortNames InSpeaker);
+AUDIOEXPERIMENTALRUNTIME_API const TCHAR* LexToString(const ESpeakerShortNames InSpeaker);
 
 namespace Audio
 {
@@ -63,7 +69,7 @@ namespace Audio
 
 	// Forward declare these for visitor.
 	class FDiscreteChannelTypeFamily;
-	class FAmbisonicsChannelTypeFamily;
+	class FSoundfieldChannelTypeFamily;
 	class FPackedChannelTypeFamily;
 	
 	class IChannelTypeVisitor
@@ -71,7 +77,7 @@ namespace Audio
 	public:
 		virtual ~IChannelTypeVisitor() = default;
 		virtual void Visit(const FDiscreteChannelTypeFamily&) = 0;
-		virtual void Visit(const FAmbisonicsChannelTypeFamily&) = 0;
+		virtual void Visit(const FSoundfieldChannelTypeFamily&) = 0;
 	};
 	
 	// Base type for all "channel" based types.
@@ -153,6 +159,16 @@ namespace Audio
 		 */
 		AUDIOEXPERIMENTALRUNTIME_API FTranscoder GetTranscoder(const FGetTranscoderParams& InParam) const;
 
+		using IPanner = IDiscretePanner;
+		
+		struct FGetPannerParams
+		{
+			// Control which algorithm to use for this format.
+			// Not currently used.
+			// FName PannerType = TEXT("VBAP 2D");
+		};
+		AUDIOEXPERIMENTALRUNTIME_API virtual const IPanner* GetPanner(const FGetPannerParams& InParams = {}) const { return nullptr; }
+		
 		/**
 		 * Returns the family name (C++ type) of this ChannelType. This allows us to safely downcast (if necessary).
 		 * @return The name of this type Family. (i.e. Discrete/Ambisoncs). 
@@ -169,7 +185,7 @@ namespace Audio
 			FName Name;
 			FString FriendlyName;
 		};
-		virtual TOptional<FChannelName> GetChannelName(const int32 InChannelIndex) const { return {}; }
+		AUDIOEXPERIMENTALRUNTIME_API virtual TOptional<FChannelName> GetChannelName(const int32 InChannelIndex) const { return {}; }
 
 		template<typename T>
 		T* Cast()
@@ -192,7 +208,12 @@ namespace Audio
 	protected:
 		// Both are friends so we can hide the api.
 		friend class FDiscreteChannelTypeFamily;
-		friend class FAmbisonicsChannelTypeFamily;
+		friend class FSoundfieldChannelTypeFamily;
+
+		void SetNumChannels(const int32 InNumChannels)
+		{
+			NumChannelsPrivate = InNumChannels;
+		}
 		
 	private:
 		bool bIsAbstract = false;
@@ -202,7 +223,7 @@ namespace Audio
 		FName FamilyType; 
 	};
 	
-	class FDiscreteChannelTypeFamily final: public FChannelTypeFamily 
+	class FDiscreteChannelTypeFamily : public FChannelTypeFamily 
 	{
 	public:
 		UE_NONCOPYABLE(FDiscreteChannelTypeFamily);
@@ -210,9 +231,10 @@ namespace Audio
 
 		struct FSpeaker
 		{
-			ESpeakerShortNames Speaker;		// Enum of speaker types.
-			float AzimuthDegrees = 0.f;
-			float ElevationDegrees = 0.f;
+			ESpeakerShortNames Speaker;				// Enum of speaker types.
+			float AzimuthDegrees = 0.f;				// -180 (left), 0 (front), 180 (right).
+			float ElevationDegrees = 0.f;			// -180 (down), 0 (ground), 180 (up). 
+			bool bIsSpatialized = true;				// If false, azimuth+elevation are ignored.
 		};
 		
 		static FName GetFamilyTypeName()
@@ -221,9 +243,20 @@ namespace Audio
 			return Name;
 		}
 			
-		FDiscreteChannelTypeFamily(const FName& InUniqueName, FChannelTypeFamily* InParentType, const FString& InFriendlyName,
+		AUDIOEXPERIMENTALRUNTIME_API FDiscreteChannelTypeFamily(const FName& InUniqueName, FChannelTypeFamily* InParentType, const FString& InFriendlyName,
 			const TArray<FSpeaker>& InOrder, const bool bIsParentsDefault, const bool bIsAbstract);
 
+		AUDIOEXPERIMENTALRUNTIME_API virtual ~FDiscreteChannelTypeFamily() override;
+		
+		/**
+		 * Returns true if this format defines any height data. 
+		 * @return true if heights are found.
+		 */
+		bool HasHeightChannels() const
+		{
+			return Order.FindByPredicate([](const FSpeaker& i) -> bool { return !FMath::IsNearlyZero(i.ElevationDegrees); }) != nullptr;
+		}
+		
 		/**
 		 * Find the index in the list of Channels for speaker.
 		 * @param InSpeaker 
@@ -250,14 +283,12 @@ namespace Audio
 		 * @return The Function Object. (this will be not be set if a function wasn't returned)
 		 */	
 
-		virtual TOptional<FChannelName> GetChannelName(const int32 InChannelIndex) const override;
+		AUDIOEXPERIMENTALRUNTIME_API virtual TOptional<FChannelName> GetChannelName(const int32 InChannelIndex) const override;
 
 		/**
 		 * Get the order of the speakers defined on this discrete format.
 		 * @return Array of speakers in order.
 		 */
-		//const TArray<ESpeakerShortNames>& GetSpeakerOrder() const { return Order; }
-
 		const TArray<FSpeaker>& GetSpeakerOrder() const { return Order; }
 
 		virtual void Accept(IChannelTypeVisitor& InVisitor) const override
@@ -265,31 +296,40 @@ namespace Audio
 			InVisitor.Visit(*this);
 		}
 
+		virtual const IPanner* GetPanner(const FGetPannerParams& InParams ={}) const override
+		{
+			return Panner.Get();	
+		}
+
 	protected:
+		static TUniquePtr<IDiscretePanner> MakePanner(const TArray<FSpeaker>& Order);
 		TArray<FSpeaker> Order;
+		TUniquePtr<IPanner> Panner;	// Just one for now.
 	};
 
-	// Ambisonics.
-	class FAmbisonicsChannelTypeFamily final : public FChannelTypeFamily 
+	// Soundfield 
+	class FSoundfieldChannelTypeFamily : public FChannelTypeFamily 
 	{
 	public:
-		UE_NONCOPYABLE(FAmbisonicsChannelTypeFamily);
+		UE_NONCOPYABLE(FSoundfieldChannelTypeFamily);
 		using Super = FChannelTypeFamily;
 		
 		static FName GetFamilyTypeName()
 		{
-			static const FName Name = TEXT("Ambisonics");
+			static const FName Name = TEXT("Soundfield");
 			return Name;
 		}
 
 		static int32 OrderToNumChannels(const int32 InOrder);
 
-		explicit FAmbisonicsChannelTypeFamily(const FName& InUniqueName, const int32 InOrder, FChannelTypeFamily* InParentType, const FString& InFriendlyName, const bool bIsParentsDefault, const bool InbIsAbstract)
+		AUDIOEXPERIMENTALRUNTIME_API explicit FSoundfieldChannelTypeFamily(const FName& InUniqueName, const int32 InOrder, FChannelTypeFamily* InParentType, const FString& InFriendlyName, const bool bIsParentsDefault, const bool InbIsAbstract)
 			: Super(InUniqueName, GetFamilyTypeName(), OrderToNumChannels(InOrder), InParentType, InFriendlyName, bIsParentsDefault, InbIsAbstract)
 			, Order(InOrder)
 		{}
 
 		int32 GetAmbisonicsOrder() const { return Order; }
+		
+		AUDIOEXPERIMENTALRUNTIME_API virtual TOptional<FChannelName> GetChannelName(const int32 InChannelIndex) const override;
 
 		virtual void Accept(IChannelTypeVisitor& InVisitor) const override
 		{
@@ -299,24 +339,24 @@ namespace Audio
 		int32 Order = 0;
 	};
 
-	/*
-	class FPackedChannelTypeFamily final : public FChannelTypeFamily
+	
+	class FCompositeChannelTypeFamily : public FChannelTypeFamily
 	{
 	public:
-		UE_NONCOPYABLE(FPackedChannelTypeFamily);
+		UE_NONCOPYABLE(FCompositeChannelTypeFamily);
 
 		using Super = FChannelTypeFamily;
 
 		static FName GetFamilyTypeName()
 		{
-			static const FName Name = TEXT("Packed");
+			static const FName Name = TEXT("Composite");
 			return Name;
 		}
-		explicit FPackedChannelTypeFamily(const FName& InUniqueName, FChannelTypeFamily* InParentType, const FString& InFriendlyName, const bool bIsParentsDefault)
+		AUDIOEXPERIMENTALRUNTIME_API explicit FCompositeChannelTypeFamily(const FName& InUniqueName, FChannelTypeFamily* InParentType, const FString& InFriendlyName, const bool bIsParentsDefault)
 			: Super(InUniqueName, GetFamilyTypeName(), 0, InParentType, InFriendlyName, bIsParentsDefault, false)
 		{}
 	};
-	*/
+	
 	
 	// Channel type registry.
 	class IChannelTypeRegistry : public ITypeFamilyRegistry
@@ -340,6 +380,8 @@ namespace Audio
 		 */
 		const FChannelTypeFamily* FindConcreteChannel(const FName InName) const
 		{
+			check(!InName.IsNone());
+
 			// Walk looking for non-abstract default child. (e.g. Surround->Surround5X->Surround 5.1)
 			const FChannelTypeFamily* Channel = FindChannel(InName);
 			while (Channel && Channel->IsAbstract())
@@ -348,12 +390,18 @@ namespace Audio
 			}
 			return Channel;
 		}
+		const FChannelTypeFamily& FindConcreteChannelChecked(const FName InName) const
+		{
+			const FChannelTypeFamily* Concrete = FindConcreteChannel(InName);
+			checkf(Concrete, TEXT("Concrete Channel not Found: %s"), *InName.ToString());
+			return *Concrete;
+		}
 
 		/**
 		 * Returns every registered format as an array.
 		 * @return Array of format types.
 		 */
-		virtual TArray<const FChannelTypeFamily*> GetAllChannelFormats() const = 0;
+		virtual TArray<TSharedRef<const FChannelTypeFamily>> GetAllChannelFormats() const = 0;
 	};
 
 	// Loose functions for now.
diff --git a/Engine/Plugins/Experimental/MetasoundExperimental/Source/AudioExperimentalRuntime/Public/TypeFamily/TypeFamily.h b/Engine/Plugins/Experimental/MetasoundExperimental/Source/AudioExperimentalRuntime/Public/TypeFamily/TypeFamily.h
--- a/Engine/Plugins/Experimental/MetasoundExperimental/Source/AudioExperimentalRuntime/Public/TypeFamily/TypeFamily.h
+++ b/Engine/Plugins/Experimental/MetasoundExperimental/Source/AudioExperimentalRuntime/Public/TypeFamily/TypeFamily.h
@@ -3,6 +3,7 @@
 #pragma once
 
 #include "Templates/RetainedRef.h"
+#include "Templates/UniquePtr.h"
 #include "UObject/NameTypes.h"
 
 namespace Audio
@@ -37,6 +38,7 @@ namespace Audio
 		bool IsA(const FName& Other) const { return nullptr != FindAncestor([&](const FTypeFamily* i) -> bool { return i->Name==Other; }); }
 		FName GetName() const { return Name; }
 		FString GetFriendlyName() const { return FriendlyName; }
+		const FTypeFamily* GetParent() const { return Parent; }
 	};
 
 	// Simple type registry.
@@ -44,7 +46,7 @@ namespace Audio
 	{
 	public:
 		virtual ~ITypeFamilyRegistry() = default;
-		virtual bool RegisterType(const FName& InUniqueName, const TRetainedRef<FTypeFamily> InType) = 0;
+		virtual bool RegisterType(const FName& InUniqueName, const TUniquePtr<FTypeFamily> InType) = 0;
 
 		template<typename T> const T* Find(const FName& InUniqueName) const
 		{
diff --git a/Engine/Plugins/Experimental/MetasoundExperimental/Source/MetasoundExperimentalEditor/MetasoundExperimentalEditor.build.cs b/Engine/Plugins/Experimental/MetasoundExperimental/Source/MetasoundExperimentalEditor/MetasoundExperimentalEditor.build.cs
--- a/Engine/Plugins/Experimental/MetasoundExperimental/Source/MetasoundExperimentalEditor/MetasoundExperimentalEditor.build.cs
+++ b/Engine/Plugins/Experimental/MetasoundExperimental/Source/MetasoundExperimentalEditor/MetasoundExperimentalEditor.build.cs
@@ -47,7 +47,8 @@ public class MetasoundExperimentalEditor : ModuleRules
 				"MetasoundEditor",
 				"StructUtilsEditor",
 				"MetasoundEngine",
-				"MetasoundFrontend"
+				"MetasoundFrontend",
+				"MetasoundGraphCore",
 				// ... add private dependencies that you statically link with here ...	
 			}
 		);
diff --git a/Engine/Plugins/Experimental/MetasoundExperimental/Source/MetasoundExperimentalEditor/Private/MetasoundExperimentalEditorModule.cpp b/Engine/Plugins/Experimental/MetasoundExperimental/Source/MetasoundExperimentalEditor/Private/MetasoundExperimentalEditorModule.cpp
--- a/Engine/Plugins/Experimental/MetasoundExperimental/Source/MetasoundExperimentalEditor/Private/MetasoundExperimentalEditorModule.cpp
+++ b/Engine/Plugins/Experimental/MetasoundExperimental/Source/MetasoundExperimentalEditor/Private/MetasoundExperimentalEditorModule.cpp
@@ -5,6 +5,8 @@
 #include "MetasoundExampleNodeDetailsCustomization.h"
 #include "MetasoundMappingFunctionNode.h"
 #include "MetasoundMappingFunctionDetailsCustomization.h"
+#include "MetasoundGranulatorNode.h"
+#include "MetasoundGranularNodeDetailsCustomization.h"
 #include "Modules/ModuleManager.h"
 
 #define LOCTEXT_NAMESPACE "MetasoundExperimentalEditor"
@@ -33,6 +35,12 @@ void FMetasoundExperimentalEditorModule::StartupModule()
 		{
 			return MakeShareable(new FMappingFunctionNodeConfigurationCustomization(InStructProperty, InNode));
 		});
+	MetasoundEditorModule.RegisterCustomNodeConfigurationDetailsCustomization(		
+		FMetaSoundGranulatorNodeConfiguration::StaticStruct()->GetFName(),
+		[](TSharedPtr<IPropertyHandle> InStructProperty, TWeakObjectPtr<UMetasoundEditorGraphNode> InNode)
+		{
+			return MakeShareable(new FGranularNodeConfigurationCustomization(InStructProperty, InNode));
+		});
 	MetasoundEditorModule.RegisterPinType("Enum:CurveFunction", "Int32");
 }
 
@@ -41,6 +49,7 @@ void FMetasoundExperimentalEditorModule::ShutdownModule()
 	Metasound::Editor::IMetasoundEditorModule& MetasoundEditorModule = FModuleManager::GetModuleChecked<Metasound::Editor::IMetasoundEditorModule>("MetasoundEditor");
 	MetasoundEditorModule.UnregisterCustomNodeConfigurationDetailsCustomization(FMetaSoundWidgetExampleNodeConfiguration::StaticStruct()->GetFName());
 	MetasoundEditorModule.UnregisterCustomNodeConfigurationDetailsCustomization(FMetaSoundMappingFunctionNodeConfiguration::StaticStruct()->GetFName());
+	MetasoundEditorModule.UnregisterCustomNodeConfigurationDetailsCustomization(FMetaSoundGranulatorNodeConfiguration::StaticStruct()->GetFName());
 }
 
 #undef LOCTEXT_NAMESPACE
diff --git a/Engine/Plugins/Experimental/MetasoundExperimental/Source/MetasoundExperimentalEditor/Private/MetasoundGranularNodeDetailsCustomization.cpp b/Engine/Plugins/Experimental/MetasoundExperimental/Source/MetasoundExperimentalEditor/Private/MetasoundGranularNodeDetailsCustomization.cpp
new file mode 100644
--- /dev/null
+++ b/Engine/Plugins/Experimental/MetasoundExperimental/Source/MetasoundExperimentalEditor/Private/MetasoundGranularNodeDetailsCustomization.cpp
@@ -0,0 +1,77 @@
+// Copyright Epic Games, Inc. All Rights Reserved.
+
+#include "MetasoundGranularNodeDetailsCustomization.h" 
+
+#include "DetailWidgetRow.h"
+#include "IDetailChildrenBuilder.h"
+#include "InstancedStructDetails.h"
+#include "MetasoundEditorGraphNode.h"
+#include "MetasoundBuilderBase.h"
+#include "MetasoundGranulatorNode.h"
+#include "PropertyEditorModule.h"
+
+#define LOCTEXT_NAMESPACE "MetasoundExperimentalEditor"
+
+const FString StructIdentifierWithDelimiters =  TEXT(".Struct.");
+
+FGranularNodeConfigurationCustomization::FGranularNodeConfigurationCustomization(TSharedPtr<IPropertyHandle> InStructProperty, TWeakObjectPtr<UMetasoundEditorGraphNode> InNode)
+	: Metasound::Editor::FMetaSoundNodeConfigurationDataDetails(InStructProperty, InNode)
+{
+	if (InStructProperty && InStructProperty->IsValidHandle())
+	{
+		StructPropertyPath = InStructProperty->GeneratePathToProperty();
+	}
+}
+
+void FGranularNodeConfigurationCustomization::OnChildRowAdded(IDetailPropertyRow& ChildRow)
+{
+	TSharedPtr<IPropertyHandle> ChildHandle = ChildRow.GetPropertyHandle();
+	if (!ChildHandle || !ChildHandle->IsValidHandle())
+	{
+		return;
+	}
+
+	const FString PropertyPath = ChildHandle->GeneratePathToProperty();
+
+	if (PropertyPath == StructPropertyPath + StructIdentifierWithDelimiters + GET_MEMBER_NAME_CHECKED(FMetaSoundGranulatorNodeConfiguration, GranularEnvelope).ToString())
+ 	{
+		EnvelopeTypePropertyHandle = ChildHandle;
+	}
+
+	// Add custom onvalue changed
+	TDelegate<void(const FPropertyChangedEvent&)> OnValueChangedDelegate = TDelegate<void(const FPropertyChangedEvent&)>::CreateSP(this, &FGranularNodeConfigurationCustomization::OnChildPropertyChanged);
+	ChildHandle->SetOnPropertyValueChangedWithData(OnValueChangedDelegate);
+
+	// Add base class on value changed
+	Metasound::Editor::FMetaSoundNodeConfigurationDataDetails::OnChildRowAdded(ChildRow);
+}
+
+void FGranularNodeConfigurationCustomization::OnChildPropertyChanged(const FPropertyChangedEvent& InPropertyChangedEvent)
+{
+	if (InPropertyChangedEvent.GetPropertyName() == GET_MEMBER_NAME_CHECKED(FMetaSoundGranulatorNodeConfiguration, GranularEnvelope))
+	{
+		if (GraphNode.IsValid() && EnvelopeTypePropertyHandle.IsValid() && EnvelopeTypePropertyHandle->IsValidHandle())
+		{
+			FMetaSoundFrontendDocumentBuilder& DocBuilder = GraphNode->GetBuilderChecked().GetBuilder();
+			const FGuid& NodeID = GraphNode->GetNodeID();
+
+			// Update the operator data value from the configuration property handle value
+			// Node configuration operator data API is experimental
+			// so this code will be made cleaner in the future 
+			TInstancedStruct<FMetaSoundFrontendNodeConfiguration> Config = DocBuilder.FindNodeConfiguration(NodeID);
+			TSharedPtr<const Metasound::IOperatorData> OperatorData = Config.Get().GetOperatorData();
+			const TSharedPtr<const Metasound::GranulatorPrivate::FGranulatorOperatorData> GranularOperatorData = StaticCastSharedPtr<const Metasound::GranulatorPrivate::FGranulatorOperatorData>(OperatorData);
+			TSharedPtr<Metasound::GranulatorPrivate::FGranulatorOperatorData> MutableOperatorData = ConstCastSharedPtr<Metasound::GranulatorPrivate::FGranulatorOperatorData>(GranularOperatorData);
+
+			TArray<void*> RawData;
+			EnvelopeTypePropertyHandle->AccessRawData(RawData);
+			if (RawData.Num() > 0)
+			{
+				EGranularEnvelope* GranularEnvelope = static_cast<EGranularEnvelope*>(RawData[0]);
+				MutableOperatorData->GranularEnvelope = *GranularEnvelope;
+			}
+		}
+	}
+}
+
+#undef LOCTEXT_NAMESPACE
diff --git a/Engine/Plugins/Experimental/MetasoundExperimental/Source/MetasoundExperimentalEditor/Private/MetasoundGranularNodeDetailsCustomization.h b/Engine/Plugins/Experimental/MetasoundExperimental/Source/MetasoundExperimentalEditor/Private/MetasoundGranularNodeDetailsCustomization.h
new file mode 100644
--- /dev/null
+++ b/Engine/Plugins/Experimental/MetasoundExperimental/Source/MetasoundExperimentalEditor/Private/MetasoundGranularNodeDetailsCustomization.h
@@ -0,0 +1,18 @@
+// Copyright Epic Games, Inc. All Rights Reserved.
+#pragma once
+
+#include "MetasoundNodeConfigurationCustomization.h"
+
+class FGranularNodeConfigurationCustomization : public Metasound::Editor::FMetaSoundNodeConfigurationDataDetails
+{
+public:
+	FGranularNodeConfigurationCustomization(TSharedPtr<IPropertyHandle> InStructProperty, TWeakObjectPtr<UMetasoundEditorGraphNode> InNode);
+
+	virtual void OnChildRowAdded(IDetailPropertyRow& ChildRow) override;
+
+private:
+	void OnChildPropertyChanged(const FPropertyChangedEvent& InPropertyChangedEvent);
+
+	FString StructPropertyPath;
+	TSharedPtr<IPropertyHandle> EnvelopeTypePropertyHandle;
+};
diff --git a/Engine/Plugins/Experimental/MetasoundExperimental/Source/MetasoundExperimentalEngineRuntime/MetasoundExperimentalEngineRuntime.Build.cs b/Engine/Plugins/Experimental/MetasoundExperimental/Source/MetasoundExperimentalEngineRuntime/MetasoundExperimentalEngineRuntime.Build.cs
--- a/Engine/Plugins/Experimental/MetasoundExperimental/Source/MetasoundExperimentalEngineRuntime/MetasoundExperimentalEngineRuntime.Build.cs
+++ b/Engine/Plugins/Experimental/MetasoundExperimental/Source/MetasoundExperimentalEngineRuntime/MetasoundExperimentalEngineRuntime.Build.cs
@@ -34,7 +34,8 @@ public class MetasoundExperimentalEngineRuntime : ModuleRules
 				"MetasoundFrontend",
 				"MetasoundEngine",
 				"MetasoundStandardNodes",
-				"SignalProcessing"
+				"SignalProcessing",
+				"MetasoundExperimentalRuntime"
 			}
 		);
 
diff --git a/Engine/Plugins/Experimental/MetasoundExperimental/Source/MetasoundExperimentalEngineRuntime/Private/MetasoundGranulatorNode.cpp b/Engine/Plugins/Experimental/MetasoundExperimental/Source/MetasoundExperimentalEngineRuntime/Private/MetasoundGranulatorNode.cpp
new file mode 100644
--- /dev/null
+++ b/Engine/Plugins/Experimental/MetasoundExperimental/Source/MetasoundExperimentalEngineRuntime/Private/MetasoundGranulatorNode.cpp
@@ -0,0 +1,991 @@
+// Copyright Epic Games, Inc. All Rights Reserved.
+
+#include "MetasoundGranulatorNode.h"
+
+#include "DSP/BufferVectorOperations.h"
+#include "DSP/FloatArrayMath.h"
+#include "DSP/GrainEnvelope.h"
+#include "DSP/SampleRateConverter.h"
+#include "MetasoundChannelAgnosticType.h"
+#include "MetasoundExecutableOperator.h"
+#include "MetasoundNodeRegistrationMacro.h"
+#include "MetasoundTrigger.h"
+#include "MetasoundVertex.h"
+#include "MetasoundWave.h"
+#include "DSP/Vbap.h"
+#include "Sound/SoundWave.h"
+#include "Sound/SoundWaveProxyReader.h"
+#include "TypeFamily/ChannelTypeFamily.h"
+#include "MetasoundCatPannerNode.h"
+
+#define LOCTEXT_NAMESPACE "MetasoundExperimentalNodes_GranulatorNode"
+
+namespace Metasound
+{
+	class FChannelAgnosticType;
+	struct FBuildOperatorParams;
+	class FNodeFacade;
+
+	namespace GranulatorPrivate
+	{
+		METASOUND_PARAM(InputSpawnGrain,	"Grain Spawn",				"Spawns a new grain from the wave asset.");
+		METASOUND_PARAM(InputStop,			"Stop",						"Stops and resets the granulator state.");
+		METASOUND_PARAM(InputWaveAsset,		"Wave Asset",				"Input sound wave to granulate.");
+		METASOUND_PARAM(InputLocationX,		"Grain Location X",			"Location in time (as a normalized fraction of audio file length) to spawn new grains from.");
+		METASOUND_PARAM(InputChannelY,		"Grain Channel Y",			"Location in channel-space (as a normalized fraction of audio file channel count) to generated blended grains.");
+		METASOUND_PARAM(InputDuration,		"Grain Duration",			"Duration (in milliseconds) of new spawned grains.");
+		METASOUND_PARAM(InputPlaybackRate, "Grain Playback Rate",		"Playback rate of new grains (1.0 is normal speed, -1.0 is backward).");
+		METASOUND_PARAM(InputVolume,		"Grain Volume",				"Linear volume scale of new grains.");
+		METASOUND_PARAM(InputPan,			"Grain Pan",				"Azimuthal pan of new grains (0.0 = left, 1.0 = right, 0.5 = center, 1.5 is behind).");
+		METASOUND_PARAM(OutputToCat,		"Output",	 				"Output audio.");
+
+		// Linkage.
+		const FLazyName FGranulatorOperatorData::OperatorDataTypeName = TEXT("FGranulatorOperatorData");
+
+		struct FDecodedDataChunk
+		{
+			// The start frame of the decoded audio chunk
+			int32 FrameStart = INDEX_NONE;
+
+			// The actual decoded audio. 
+			// Size is NumFrames * NumChannels * MetaSoundSampleRate (after SRC the source file)
+			Audio::FAlignedFloatBuffer PCMAudio;
+
+			// The number of audio frames in the data chunk
+			int32 NumFramesInChunk = 0;
+
+			// Count of the number of grains actively using this chunk
+			int32 NumGrainsUsingChunk = 0;
+		};
+
+		// Data used for spawning new grains
+		struct FSpawnGrainData
+		{
+			int32 StartFrameOffset = 0;
+			float LocationTimeSeconds = 0.0f;
+			float LocationChannelFraction = 0.5f;
+			int32 DurationFrames = 100;
+			float PlaybackRate = 1.0f;
+			float Volume = 1.0f;
+			float Pan = 0.5;
+			EGranularEnvelope Envelope = EGranularEnvelope::Gaussian;
+		};
+
+		// Active grain data for a single grain
+		struct FGrain
+		{
+			// The frame this grain starts rendering in the current block. Allows grains to schedule mid-block at the top of the render block.
+			int32 FrameStartOffset = 0;
+			
+			// The number of frames this grain has rendered
+			// If this is larger than GrainDurationFrames it's an inactive grain
+			// this is used to look up the grain envelope to find the amplitude of the grain per frame
+			int32 CurrentRenderedFramesCount = INDEX_NONE;
+
+			// Index into the decoded data array that this grain is currently using
+			int32 DecodedDataChunkIndex = 0;
+
+			// The current read frame of the grain. Float to support dynamic playback rate.
+			float CurrentReadFrame = INDEX_NONE;
+			float CurrentPlaybackRate = 1.0f;
+
+			// Previous sample values (from previous decoded chunk) for the Y-channel reading feature
+			float PrevASampleValue = 0.0f;
+			float PrevBSampleValue = 0.0f;
+
+			// The duration of this grain. Set when the grain spawns. 
+			int32 GrainDurationFrames = 0;
+
+			// The current channel position of the grain (interpolated across the audio channels)
+			float CurrentLocationChannelFraction = 0.5f;
+
+			// Grain panning information
+			float CurrentPanParameter = 0.5f;
+			float LastUpdatePanningFraction = -1.0f;
+			float LastUpdatePrevChannelGain = -1.0f;
+			float LastUpdateNextChannelGain = -1.0f;
+
+			// Grain volume information
+			float CurrentVolume = 1.0f;
+			float PreviousVolume = 1.0f;
+
+			EGranularEnvelope EnvelopType = EGranularEnvelope::Gaussian;
+		};
+	}
+
+	static const Audio::FDiscreteChannelTypeFamily& GetDiscreteFamilyChecked(const Audio::FChannelTypeFamily& FamilyType)
+	{
+		check(FamilyType.GetFamilyName() == Audio::FDiscreteChannelTypeFamily::GetFamilyTypeName());
+		return static_cast<const Audio::FDiscreteChannelTypeFamily&>(FamilyType);
+	}
+
+	class FGranulatorOperator final : public TExecutableOperator<FGranulatorOperator>
+	{
+	public:
+		using FDiscreteChannelTypeFamily = Audio::FDiscreteChannelTypeFamily;
+		using FGranulatorOperatorData = GranulatorPrivate::FGranulatorOperatorData;
+		using FDecodedDataChunk = GranulatorPrivate::FDecodedDataChunk;
+		using FGrain = GranulatorPrivate::FGrain;
+		using FSpawnGrainData = GranulatorPrivate::FSpawnGrainData;
+
+		FGranulatorOperator(
+				const TSharedPtr<const FGranulatorOperatorData>& InOperatorData,
+				const FBuildOperatorParams& InParams, 
+				FTriggerReadRef&& InSpawnGrain,
+				FTriggerReadRef&& InStop,
+				FWaveAssetReadRef&& InWaveAsset,
+				FFloatReadRef&& InGrainLocationX,
+				FFloatReadRef&& InGrainChannelY,
+				FFloatReadRef&& InGrainDuration,
+				FFloatReadRef&& InGrainPlaybackRate,
+				FFloatReadRef&& InGrainVolume,
+				FFloatReadRef&& InGrainPan,
+				FDiscreteChannelAgnosticTypeWriteRef&& InOutputAudio
+				)
+			: OperatorData(InOperatorData)
+			, SpawnGrainTrigger(MoveTemp(InSpawnGrain))
+			, StopTrigger(MoveTemp(InStop))
+			, WaveAsset(MoveTemp(InWaveAsset))
+			, GrainLocationX(MoveTemp(InGrainLocationX))
+			, GrainChannelY(MoveTemp(InGrainChannelY))
+			, GrainDuration(MoveTemp(InGrainDuration))
+			, GrainPlaybackRate(MoveTemp(InGrainPlaybackRate))
+			, GrainVolume(MoveTemp(InGrainVolume))
+			, GrainPan(MoveTemp(InGrainPan))
+			, OutputAudio(MoveTemp(InOutputAudio))
+			, Settings(InParams.OperatorSettings)
+			, DiscreteFamily(GetDiscreteFamilyChecked(OutputAudio->GetType()))
+			, SRC(Audio::ISampleRateConverter::CreateSampleRateConverter())
+			, GranularEnvelopeType(InOperatorData->GranularEnvelope)
+			, Panner(OutputAudio->GetType().GetPanner())
+		{
+			// Generate the grain envelope data for the different grain envelopes
+			int32 NumEnvelopeTypes = (int32)EGranularEnvelope::Hann + 1;
+			GrainEnvelopeData.AddDefaulted(NumEnvelopeTypes);
+			for (int32 i = 0; i < NumEnvelopeTypes; ++i)
+			{
+				Audio::Grain::EEnvelope EnvType = (Audio::Grain::EEnvelope)i;
+				Audio::Grain::GenerateEnvelopeData(GrainEnvelopeData[i], 256, EnvType);
+			}
+		}
+		virtual ~FGranulatorOperator() override = default;
+
+		static const FVertexInterface MakeInterface(const FName InOutputFormat)
+		{
+			using namespace GranulatorPrivate;
+			// inputs
+			FInputVertexInterface InputInterface;
+			InputInterface.Add(TInputDataVertex<FWaveAsset>(METASOUND_GET_PARAM_NAME_AND_METADATA(InputWaveAsset)));
+			InputInterface.Add(TInputDataVertex<FTrigger>(METASOUND_GET_PARAM_NAME_AND_METADATA(InputSpawnGrain)));
+			InputInterface.Add(TInputDataVertex<FTrigger>(METASOUND_GET_PARAM_NAME_AND_METADATA(InputStop)));
+			InputInterface.Add(TInputDataVertex<float>(METASOUND_GET_PARAM_NAME_AND_METADATA(InputLocationX), 0.0f));
+			InputInterface.Add(TInputDataVertex<float>(METASOUND_GET_PARAM_NAME_AND_METADATA(InputChannelY), 0.0f));
+			InputInterface.Add(TInputDataVertex<float>(METASOUND_GET_PARAM_NAME_AND_METADATA(InputDuration), 100.0f));
+			InputInterface.Add(TInputDataVertex<float>(METASOUND_GET_PARAM_NAME_AND_METADATA(InputPlaybackRate), 1.0f));
+			InputInterface.Add(TInputDataVertex<float>(METASOUND_GET_PARAM_NAME_AND_METADATA(InputVolume), 1.0f));
+			InputInterface.Add(TInputDataVertex<float>(METASOUND_GET_PARAM_NAME_AND_METADATA(InputPan), 0.5f));
+
+			// outputs
+			FOutputVertexInterface OutputInterface;
+			OutputInterface.Add(FOutputDataVertex(METASOUND_GET_PARAM_NAME(OutputToCat), InOutputFormat, METASOUND_GET_PARAM_METADATA(OutputToCat), EVertexAccessType::Reference));
+			
+			return FVertexInterface(InputInterface, OutputInterface);
+		}
+			
+		static TUniquePtr<IOperator> CreateOperator(const FBuildOperatorParams& InParams, FBuildResults& OutResults)
+		{
+			using namespace GranulatorPrivate;
+
+			const FGranulatorOperatorData* ConfigData = CastOperatorData<const FGranulatorOperatorData>(InParams.Node.GetOperatorData().Get());
+			if (!ConfigData)
+			{
+				return MakeUnique<FNoOpOperator>();
+			}
+			const FName AudioOutputToName = ConfigData->OutputAudioTypeName;
+			const Audio::FChannelTypeFamily* ConcretePanToType = Audio::GetChannelRegistry().FindConcreteChannel(AudioOutputToName);
+			if (!ConcretePanToType)
+			{
+				return MakeUnique<FNoOpOperator>();
+			}
+
+			const TSharedPtr<const FGranulatorOperatorData>& OperatorDataSharedPtr = StaticCastSharedPtr<const FGranulatorOperatorData>(InParams.Node.GetOperatorData());
+
+			TDataReadReference<FTrigger> InputSpawnGrainRef = InParams.InputData.GetOrCreateDefaultDataReadReference<FTrigger>(METASOUND_GET_PARAM_NAME(InputSpawnGrain), InParams.OperatorSettings);
+			TDataReadReference<FTrigger> InputStopRef = InParams.InputData.GetOrCreateDefaultDataReadReference<FTrigger>(METASOUND_GET_PARAM_NAME(InputStop), InParams.OperatorSettings);
+			TDataReadReference<FWaveAsset> InputWaveAssetRef = InParams.InputData.GetOrCreateDefaultDataReadReference<FWaveAsset>(METASOUND_GET_PARAM_NAME(InputWaveAsset), InParams.OperatorSettings);
+			TDataReadReference<float> InputLocationXRef = InParams.InputData.GetOrCreateDefaultDataReadReference<float>(METASOUND_GET_PARAM_NAME(InputLocationX), InParams.OperatorSettings);
+			TDataReadReference<float> InputChannelYRef = InParams.InputData.GetOrCreateDefaultDataReadReference<float>(METASOUND_GET_PARAM_NAME(InputChannelY), InParams.OperatorSettings);
+			TDataReadReference<float> InputDurationRef = InParams.InputData.GetOrCreateDefaultDataReadReference<float>(METASOUND_GET_PARAM_NAME(InputDuration), InParams.OperatorSettings);
+			TDataReadReference<float> InputPlaybackRateRef = InParams.InputData.GetOrCreateDefaultDataReadReference<float>(METASOUND_GET_PARAM_NAME(InputPlaybackRate), InParams.OperatorSettings);
+			TDataReadReference<float> InputVolumeRef = InParams.InputData.GetOrCreateDefaultDataReadReference<float>(METASOUND_GET_PARAM_NAME(InputVolume), InParams.OperatorSettings);
+			TDataReadReference<float> InputPanRef = InParams.InputData.GetOrCreateDefaultDataReadReference<float>(METASOUND_GET_PARAM_NAME(InputPan), InParams.OperatorSettings);
+
+
+			// Make Discrete CAT output.
+			// These are defined in data so we must use "DerivedAs" 
+			TDataWriteReference<FDiscreteChannelAgnosticType> CatPin =
+				FDiscreteChannelAgnosticTypeWriteRef::CreateNewDerivedAs<FDiscreteChannelAgnosticType>(
+					ConcretePanToType->GetName(),
+					Metasound::GetMetasoundDataTypeId<FDiscreteChannelAgnosticType>(),
+					InParams.OperatorSettings,
+					ConcretePanToType->GetName()
+				);
+			
+			return MakeUnique<FGranulatorOperator>(
+				OperatorDataSharedPtr,
+				InParams,
+				MoveTemp(InputSpawnGrainRef),
+				MoveTemp(InputStopRef),
+				MoveTemp(InputWaveAssetRef),
+				MoveTemp(InputLocationXRef),
+				MoveTemp(InputChannelYRef),
+				MoveTemp(InputDurationRef),
+				MoveTemp(InputPlaybackRateRef),
+				MoveTemp(InputVolumeRef),
+				MoveTemp(InputPanRef),
+				MoveTemp(CatPin)
+			);
+		}
+
+		virtual void BindInputs(FInputVertexInterfaceData& InOutVertexData) override
+		{
+			using namespace GranulatorPrivate;
+			InOutVertexData.BindReadVertex(METASOUND_GET_PARAM_NAME(InputSpawnGrain), SpawnGrainTrigger);
+			InOutVertexData.BindReadVertex(METASOUND_GET_PARAM_NAME(InputStop), StopTrigger);
+			InOutVertexData.BindReadVertex(METASOUND_GET_PARAM_NAME(InputWaveAsset), WaveAsset);
+			InOutVertexData.BindReadVertex(METASOUND_GET_PARAM_NAME(InputLocationX), GrainLocationX);
+			InOutVertexData.BindReadVertex(METASOUND_GET_PARAM_NAME(InputChannelY), GrainChannelY);
+			InOutVertexData.BindReadVertex(METASOUND_GET_PARAM_NAME(InputDuration), GrainDuration);
+			InOutVertexData.BindReadVertex(METASOUND_GET_PARAM_NAME(InputPlaybackRate), GrainPlaybackRate);
+			InOutVertexData.BindReadVertex(METASOUND_GET_PARAM_NAME(InputVolume), GrainVolume);
+			InOutVertexData.BindReadVertex(METASOUND_GET_PARAM_NAME(InputPan), GrainPan);
+		}
+		
+		virtual void BindOutputs(FOutputVertexInterfaceData& InOutVertexData) override
+		{
+			using namespace GranulatorPrivate;
+			InOutVertexData.BindWriteVertex(METASOUND_GET_PARAM_NAME(OutputToCat), OutputAudio);
+		}
+
+		void Reset(const FResetParams& InParams)
+		{
+			Execute();
+		}
+
+		bool Initialize()
+		{
+			check(!bPlaying);
+
+ 			CurrentWaveAsset = *WaveAsset;
+			SoundWaveProxyPtr = CurrentWaveAsset.GetSoundWaveProxy();
+			if (!SoundWaveProxyPtr)
+			{
+				return false;
+			}
+
+			SourceFileSampleRate = SoundWaveProxyPtr->GetSampleRate();
+			SourceFileDurationSeconds = SoundWaveProxyPtr->GetDuration();
+			SourceFileNumFrames = SoundWaveProxyPtr->GetNumFrames();
+			SourceFileNumChannels = SoundWaveProxyPtr->GetNumChannels();
+
+			FSoundWaveProxyReader::FSettings ProxyReaderSettings;
+			ProxyReaderSettings.MaxDecodeSizeInFrames = DecodedAudioSizeInSeconds * SourceFileSampleRate;
+			ProxyReaderSettings.StartTimeInSeconds = 0.0f;
+			ProxyReaderSettings.bIsLooping = true;
+			ProxyReaderSettings.LoopDurationInSeconds = SourceFileDurationSeconds;
+			SoundWaveProxyReaderPtr = FSoundWaveProxyReader::Create(SoundWaveProxyPtr.ToSharedRef(), ProxyReaderSettings);
+
+			DecodedChunks.Reset();
+			DecodedChunks.AddDefaulted(3);
+
+			// Scratch buffer to write the audio generated from mono-grains before applying panning
+			MonoGrainBuffer.Reset();
+			MonoGrainBuffer.AddZeroed(Settings.GetNumFramesPerBlock());
+
+			return true;
+		}
+
+		void SpawnGrain(const FSpawnGrainData& InSpawnGrainData)
+		{
+			// Try to retrieve a decoded data chunk for the current read frame based on the curret, interpolated playhead time seconds value
+			int32 CurrentReadFrame = InSpawnGrainData.LocationTimeSeconds * Settings.GetSampleRate();
+			int32 DecodedDataChunkIndex = GetDecodedDataChunkIndexForCurrentReadIndex(CurrentReadFrame);
+			if (DecodedDataChunkIndex == INDEX_NONE)
+			{
+				DecodedDataChunkIndex = DecodeDataChunkIndexForCurrentReadIndex(CurrentReadFrame);
+			}
+			check(DecodedDataChunkIndex != INDEX_NONE);
+
+			// Get the grain runtime data for the new grain spawn
+			FGrain NewGrain;
+			NewGrain.FrameStartOffset = InSpawnGrainData.StartFrameOffset;
+			NewGrain.CurrentRenderedFramesCount = 0;
+			NewGrain.DecodedDataChunkIndex = DecodedDataChunkIndex;
+			NewGrain.CurrentReadFrame = (float)CurrentReadFrame;
+			NewGrain.CurrentPlaybackRate = InSpawnGrainData.PlaybackRate;
+			NewGrain.GrainDurationFrames = InSpawnGrainData.DurationFrames;
+			NewGrain.CurrentLocationChannelFraction = InSpawnGrainData.LocationChannelFraction;
+			NewGrain.CurrentPanParameter = InSpawnGrainData.Pan;
+			NewGrain.CurrentVolume = InSpawnGrainData.Volume;
+			NewGrain.EnvelopType = GranularEnvelopeType;
+
+			ActiveGrains.Add(NewGrain);
+
+			DecodedChunks[DecodedDataChunkIndex].NumGrainsUsingChunk++;
+			GrainCount++;
+		}
+
+		void DecodeToDataChunk(FDecodedDataChunk& InOutDataChunk, float InDecoderSeekTimeSeconds)
+		{
+			check(InOutDataChunk.NumGrainsUsingChunk == 0);
+			check(DecodedAudioSizeInSeconds > 0.0f);
+			check(SourceFileSampleRate > 0.0f);
+			check(InDecoderSeekTimeSeconds >= 0.0f);
+			check(SourceFileNumChannels > 0);
+			check(SoundWaveProxyPtr.IsValid());
+
+			check(SoundWaveProxyReaderPtr.IsValid());
+			float DecoderSeekTimeSeconds = FMath::Max(InDecoderSeekTimeSeconds - 0.5f * DecodedAudioSizeInSeconds, 0.0f);
+			// If proxy reader already exists, then simply seek the decoder to the desired location
+			SoundWaveProxyReaderPtr->SeekToTime(DecoderSeekTimeSeconds);
+
+			float AudioMixerSampleRate = Settings.GetSampleRate();
+
+			InOutDataChunk.FrameStart = DecoderSeekTimeSeconds * AudioMixerSampleRate;
+
+			// We allocate the size of buffer pre-SRC based on source file sample rate
+			int32 DecodedAudioSize = DecodedAudioSizeInSeconds * SourceFileSampleRate * SourceFileNumChannels;
+			check(DecodedAudioSize > 0);
+			InOutDataChunk.PCMAudio.Reset();
+			InOutDataChunk.PCMAudio.AddUninitialized(DecodedAudioSize);
+
+			// This does the actual decoding of the audio to match the size of the input buffer
+			SoundWaveProxyReaderPtr->PopAudio(InOutDataChunk.PCMAudio);
+
+			// Check if we need to do SRC. If we do, this will change the allocated 
+			// size (potentially expanding or shrinking) to match 
+			// the audio mixer sample rate.
+			if (!FMath::IsNearlyEqual(SourceFileSampleRate, AudioMixerSampleRate))
+			{
+				SRC->Init((float)SourceFileSampleRate / AudioMixerSampleRate, SourceFileNumChannels);
+				TArray<float> SampleRateConvertedPCM;
+				SRC->ProcessFullbuffer(InOutDataChunk.PCMAudio.GetData(), InOutDataChunk.PCMAudio.Num(), SampleRateConvertedPCM);
+				InOutDataChunk.PCMAudio = MoveTemp(SampleRateConvertedPCM);
+			}
+
+			InOutDataChunk.NumFramesInChunk = InOutDataChunk.PCMAudio.Num() / SourceFileNumChannels;
+		}
+
+		int32 GetDecodedDataChunkIndexForCurrentReadIndex(int32 InReadFrameIndex)
+		{
+			for (int32 i = 0; i < DecodedChunks.Num(); ++i)
+			{
+				const FDecodedDataChunk& DecodedChunk = DecodedChunks[i];
+				if (DecodedChunk.PCMAudio.Num() > 0)
+				{
+					int32 DecodedFrameCount = DecodedChunk.NumFramesInChunk;
+					if (InReadFrameIndex >= DecodedChunk.FrameStart && InReadFrameIndex < DecodedChunk.FrameStart + DecodedFrameCount)
+					{
+						// We found a decoded audio chunk that contains the desired read frame index
+						return i;
+					}
+				}
+			}
+			return INDEX_NONE;
+		}
+
+		int32 DecodeDataChunkIndexForCurrentReadIndex(int32 InReadFrameIndex)
+		{
+			float AudioMixerSampleRate = Settings.GetSampleRate();
+
+			// No decoded chunk was found for desired read frame index. 
+			// This indicates that we need to decode more audio.
+			for (int32 i = 0; i < DecodedChunks.Num(); ++i)
+			{
+				FDecodedDataChunk& DecodedChunk = DecodedChunks[i];
+				if (!DecodedChunk.NumGrainsUsingChunk)
+				{
+					float DecoderSeekTimeSeconds = InReadFrameIndex / AudioMixerSampleRate;
+					DecodeToDataChunk(DecodedChunk, DecoderSeekTimeSeconds);
+					check(DecodedChunk.PCMAudio.Num() > 0);
+					return i;
+				}
+			}
+
+			FDecodedDataChunk NewChunk;
+			float DecoderSeekTimeSeconds = InReadFrameIndex / AudioMixerSampleRate;
+			DecodeToDataChunk(NewChunk, DecoderSeekTimeSeconds);
+			check(NewChunk.PCMAudio.Num() > 0);
+			DecodedChunks.Add(MoveTemp(NewChunk));
+			return DecodedChunks.Num() - 1;
+		}
+
+		
+		void GetOutputSpeakerChannelData(FGrain& InGrain, float& OutPanningFraction, Audio::IDiscretePanner::FPanResult& OutPrev, Audio::IDiscretePanner::FPanResult& OutNext)
+		{ 
+			check(InGrain.CurrentPanParameter >= 0.0f && InGrain.CurrentPanParameter < 2.0f);
+			const float AzimuthDegrees = NormalizedAzimuthToDegrees(InGrain.CurrentPanParameter);
+
+			if (Panner)
+			{
+				Audio::IDiscretePanner::FOutputParams OutParams;
+				Panner->ComputeGains( { .AzimuthDegrees = AzimuthDegrees, .ElevationDegrees = 0, .bAllowAzimuthMirroring = true }, OutParams);
+				OutPrev = OutParams.Results[0];
+				OutNext = OutParams.Results[1];
+			}
+			
+#if 0
+			float AzimuthDegrees = FMath::Fmod((InGrain.CurrentPanParameter * 180.0f) + 270.0f, 360.0f);
+			int32 NumOutputChannels = ChannelAzimuthData.Num();
+			for (int32 i = 0; i < NumOutputChannels; ++i)
+			{
+				const FAzimuthalSpeakerData& PrevData = ChannelAzimuthData[i];
+				const FAzimuthalSpeakerData& NextData = ChannelAzimuthData[(i + 1) % NumOutputChannels];
+
+				float AzimuthDelta = FMath::Fmod(NextData.Azimuth - PrevData.Azimuth + 360.0f, 360.0f);
+				float RelAzimuth = FMath::Fmod(AzimuthDegrees - PrevData.Azimuth + 360.0f, 360.0f);
+
+				if (RelAzimuth >= 0.0f && RelAzimuth < AzimuthDelta)
+				{
+					float PanningFraction = AzimuthDelta > KINDA_SMALL_NUMBER ? (RelAzimuth / AzimuthDelta) : 0.0f;
+					OutPrev = PrevData;
+					OutNext = NextData;
+					OutPanningFraction = PanningFraction;
+					return;
+				}
+			}
+#endif //0			
+		}
+		
+
+		//void ComputePanGains(float InPanningFraction, float& OutPrevGain, float& OutNextGain) const
+		//{
+		//	check(InPanningFraction >= 0.0f && InPanningFraction <= 1.0f);
+		//
+		//FMath::SinCos(&OutNextGain, &OutPrevGain, 0.5f * PI * InPanningFraction);
+		//}
+
+		void MixMonoGrainToOutput(TArrayView<float>& InMonoBufferView, FGrain& InGrain)
+		{
+			//FAzimuthalSpeakerData PrevData;
+			//FAzimuthalSpeakerData NextData;
+			Audio::IDiscretePanner::FPanResult PrevData;
+			Audio::IDiscretePanner::FPanResult NextData;
+			float PanningFraction = 0.0f;
+			GetOutputSpeakerChannelData(InGrain, PanningFraction, PrevData, NextData);
+
+			if (PrevData.ChannelID == INDEX_NONE || NextData.ChannelID == INDEX_NONE)
+			{
+				return;
+			}
+			int32 PrevSpeakerIndex = DiscreteFamily.FindSpeakerIndex(static_cast<ESpeakerShortNames>(PrevData.ChannelID));
+			check(PrevSpeakerIndex != INDEX_NONE);
+
+			int32 NextSpeakerIndex = DiscreteFamily.FindSpeakerIndex(static_cast<ESpeakerShortNames>(NextData.ChannelID));
+			check(NextSpeakerIndex != INDEX_NONE);
+
+			TArrayView<float> PrevChannelOutAudioBufferView = OutputAudio->GetChannel(PrevSpeakerIndex);
+			TArrayView<float> NextChannelOutAudioBufferView = OutputAudio->GetChannel(NextSpeakerIndex);
+
+			// We use optimized panning algorithm that interpolates the panning amount through the block
+			float StartingPanningFraction = InGrain.LastUpdatePanningFraction < 0.0f ? PanningFraction : InGrain.LastUpdatePanningFraction;
+			float EndingPanningFraction = PanningFraction;
+
+			// This is the first execute so we haven't yet cached our prev/next channel pans
+			//if (InGrain.LastUpdatePanningFraction < 0.0f)
+			//{
+			//	ComputePanGains(StartingPanningFraction, InGrain.LastUpdatePrevChannelGain, InGrain.LastUpdateNextChannelGain);
+			//}
+
+			float CurrentPrevChannelGain = PrevData.Gain;
+			float CurrentNextChannelGain = NextData.Gain;
+			//ComputePanGains(EndingPanningFraction, CurrentPrevChannelGain, CurrentNextChannelGain);
+
+			Audio::ArrayMixIn(InMonoBufferView, PrevChannelOutAudioBufferView, InGrain.LastUpdatePrevChannelGain, CurrentPrevChannelGain);
+			Audio::ArrayMixIn(InMonoBufferView, NextChannelOutAudioBufferView, InGrain.LastUpdateNextChannelGain, CurrentNextChannelGain);
+
+			// lerp through the buffer to the target panning amount
+			InGrain.LastUpdatePanningFraction = PanningFraction;
+			InGrain.LastUpdatePrevChannelGain = CurrentPrevChannelGain;
+			InGrain.LastUpdateNextChannelGain = CurrentNextChannelGain;
+		}
+
+		void GetNewDecodedData(FGrain& InGrain, int32 InFrameOffset)
+		{
+			// We're no longer using this decoded audio chunk
+			DecodedChunks[InGrain.DecodedDataChunkIndex].NumGrainsUsingChunk--;
+			check(DecodedChunks[InGrain.DecodedDataChunkIndex].NumGrainsUsingChunk >= 0);
+
+			// If we're calling this function, that means that we ran out of data for the current read frame chunk, so load up the chunk for the next frame
+			InGrain.DecodedDataChunkIndex = GetDecodedDataChunkIndexForCurrentReadIndex((int32)InGrain.CurrentReadFrame + InFrameOffset);
+			if (InGrain.DecodedDataChunkIndex == INDEX_NONE)
+			{
+				InGrain.DecodedDataChunkIndex = DecodeDataChunkIndexForCurrentReadIndex((int32)InGrain.CurrentReadFrame + InFrameOffset);
+			}
+			check(InGrain.DecodedDataChunkIndex != INDEX_NONE);
+
+			DecodedChunks[InGrain.DecodedDataChunkIndex].NumGrainsUsingChunk++;
+		}
+
+		void RenderActiveGrains(int32 InNumFramesToRender)
+		{
+			check(SourceFileNumChannels > 0);
+
+			// Reverse iterate so we can quickly remove grains when they're done
+			for (int32 GrainIndex = ActiveGrains.Num() - 1; GrainIndex >= 0; --GrainIndex)
+			{
+				FGrain& Grain = ActiveGrains[GrainIndex];
+
+				// Cache the grain envelope ref for this grain
+				Audio::Grain::FEnvelope& GrainEnv = GrainEnvelopeData[(int32)Grain.EnvelopType];
+
+				// This is a frame offset into the current block to account for grains starting mid-block
+				int32 GrainWriteIndex = Grain.FrameStartOffset;
+				int32 NumFramesToRenderForBlock = InNumFramesToRender - GrainWriteIndex;
+
+				// After this block render, we won't need to use this frame offset
+				// It is only ever non-zero for newly spawned grains
+				Grain.FrameStartOffset = 0;
+
+				// This is the number of frames we have left to render for this grain
+				int32 NumFramesRemainingInGrain = Grain.GrainDurationFrames - Grain.CurrentRenderedFramesCount;
+
+				// We only need to render frames for either how many are in this block or how many are left in the grain to render
+				NumFramesToRenderForBlock = FMath::Min(NumFramesToRenderForBlock, NumFramesRemainingInGrain);
+				check(NumFramesToRenderForBlock > 0);
+
+				// For inner-loop audio rendering, we avoid the constant array-access checks that happen on our containers
+				FMemory::Memzero(MonoGrainBuffer.GetData(), MonoGrainBuffer.Num() * sizeof(float));
+				float* MonoGrainBufferPtr = MonoGrainBuffer.GetData();
+
+				// Calculate which 2 channels we need blend for the output
+				float ChannelIndex = Grain.CurrentLocationChannelFraction * (SourceFileNumChannels - 1);
+				int32 ChannelAIndex = (int32)ChannelIndex;
+				int32 ChannelBIndex = (ChannelAIndex + 1) % SourceFileNumChannels;
+				// The alpha blend between channels A and B
+				float ChannelAlpha = 1.0f - (ChannelIndex - (float)ChannelAIndex);
+
+				bool bGrainFinished = false;
+				while (NumFramesToRenderForBlock > 0 && !bGrainFinished)
+				{
+					check(Grain.DecodedDataChunkIndex != INDEX_NONE);
+
+					// Make sure we have a valid decode data chunk ready for rendering
+					const FDecodedDataChunk* DecodedData = &DecodedChunks[Grain.DecodedDataChunkIndex];
+
+					if (Grain.CurrentPlaybackRate > 0.0f)
+					{
+						// The number of frames that this grain is offset from the decoded data chunk
+						// We are rounding up to nearest integer since float-frame read indices require us to read to the next frame to interpolate.
+						// We should be storing the previous audio data so we can't discard the old chunk if our "next" frame read is off this current chunk
+						
+						int32 NumFramesOffsetInDecodedData = (int32)(Grain.CurrentReadFrame) - DecodedData->FrameStart + 1;
+
+						// These are the total number of frames left in the current chunk we can consume
+						// Note this can be negative if the pitch scale is > 1.0
+						int32 NumFramesPossibleToConsumeInCurrentChunk = DecodedData->NumFramesInChunk - NumFramesOffsetInDecodedData;
+
+						// If we've totally consumed this decoded audio chunk, we need to get a new decoded audio chunk
+						if (NumFramesPossibleToConsumeInCurrentChunk <= 0)
+						{
+							GetNewDecodedData(Grain, 1);
+
+							DecodedData = &DecodedChunks[Grain.DecodedDataChunkIndex];
+						}
+					}
+					else
+					{
+						if (Grain.CurrentReadFrame < (float)DecodedData->FrameStart)
+						{
+							// Wrap the read frame around the audio file
+							if (Grain.CurrentReadFrame < 0.0f)
+							{
+								Grain.CurrentReadFrame += (float)SourceFileNumFrames;
+							}
+
+							GetNewDecodedData(Grain, -1);
+
+							DecodedData = &DecodedChunks[Grain.DecodedDataChunkIndex];
+							check(Grain.CurrentReadFrame >= (float)DecodedData->FrameStart);
+						}
+					}
+
+					int32 FrameWriteIndex = 0;
+
+					while (NumFramesToRenderForBlock > 0 && !bGrainFinished)
+					{
+						float EnvelopeFraction = FMath::Clamp((float)Grain.CurrentRenderedFramesCount / Grain.GrainDurationFrames, 0.0f, 1.0f);
+						float GrainAmplitude = Audio::Grain::GetValue(GrainEnv, EnvelopeFraction);
+							
+						// Need to interpolate from the previoux block
+						float DecodedSampleDataPrev = 0.0f;
+						float DecodedSampleDataNext = 0.0f;
+						float ReadFrameAlpha = 0;
+
+						float ReadFrameOffset = Grain.CurrentReadFrame - (float)DecodedData->FrameStart;
+
+						// Playing forward
+						if (Grain.CurrentPlaybackRate > 0.0f)
+						{
+							// We're still reading from the previous block of decoded audio (which is behind our current chunk)
+							if (ReadFrameOffset < 0.0f)
+							{
+								check(ReadFrameOffset >= -1.0f);
+								DecodedSampleDataPrev = Grain.PrevASampleValue;
+								if (SourceFileNumChannels == 1)
+								{
+									DecodedSampleDataNext = DecodedData->PCMAudio[0];
+								}
+								else
+								{
+									float SampleA = DecodedData->PCMAudio[ChannelAIndex];
+									float SampleB = DecodedData->PCMAudio[ChannelBIndex];
+									DecodedSampleDataNext = FMath::Lerp(SampleA, SampleB, ChannelAlpha);
+								}
+								ReadFrameAlpha = 1.0f + ReadFrameOffset;
+								check(ReadFrameAlpha >= 0.0f && ReadFrameAlpha <= 1.0f);
+							}
+							else
+							{
+								int32 PrevDataReadIndex = (int32)ReadFrameOffset;
+								int32 NextDataReadIndex = PrevDataReadIndex + 1;
+
+								// We've consumed all the data from this audio chunk
+								if (NextDataReadIndex >= DecodedData->NumFramesInChunk)
+								{
+									break;
+								}
+
+								ReadFrameAlpha = ReadFrameOffset - (float)PrevDataReadIndex;
+
+								if (SourceFileNumChannels == 1)
+								{
+									DecodedSampleDataPrev = DecodedData->PCMAudio[PrevDataReadIndex];
+									DecodedSampleDataNext = DecodedData->PCMAudio[NextDataReadIndex];
+								}
+								else
+								{
+									int32 SampleReadIndex = PrevDataReadIndex * SourceFileNumChannels;
+									float SampleA = DecodedData->PCMAudio[SampleReadIndex + ChannelAIndex];
+									float SampleB = DecodedData->PCMAudio[SampleReadIndex + ChannelBIndex];
+
+									DecodedSampleDataPrev = FMath::Lerp(SampleA, SampleB, ChannelAlpha);
+
+									SampleReadIndex = NextDataReadIndex * SourceFileNumChannels;
+									SampleA = DecodedData->PCMAudio[SampleReadIndex + ChannelAIndex];
+									SampleB = DecodedData->PCMAudio[SampleReadIndex + ChannelBIndex];
+
+									DecodedSampleDataNext = FMath::Lerp(SampleA, SampleB, ChannelAlpha);
+								}
+							}
+						}
+						// Playing backward!
+						else
+						{
+							// We're still reading from the previous block of decoded audio data (which is in front our current chunk)
+							if (((int32)ReadFrameOffset + 1) >= DecodedData->NumFramesInChunk)
+							{
+								// We should only be 1 frame into the next block
+								DecodedSampleDataPrev = Grain.PrevASampleValue;
+								if (SourceFileNumChannels == 1)
+								{
+									DecodedSampleDataNext = DecodedData->PCMAudio.Last();
+								}
+								else
+								{									
+									int32 SampleReadIndex = (DecodedData->NumFramesInChunk - 1) * SourceFileNumChannels;
+									float SampleA = DecodedData->PCMAudio[SampleReadIndex + ChannelAIndex];
+									float SampleB = DecodedData->PCMAudio[SampleReadIndex + ChannelBIndex];
+									DecodedSampleDataNext = FMath::Lerp(SampleA, SampleB, ChannelAlpha);
+								}
+
+								float Delta = FMath::Clamp((float)DecodedData->NumFramesInChunk - ReadFrameOffset, 0.0f, 1.0f);
+								ReadFrameAlpha = 1.0f - Delta;
+							}
+							else
+							{
+								int32 NextDataReadIndex = (int32)ReadFrameOffset;
+								int32 PrevDataReadIndex = NextDataReadIndex + 1;
+
+								ReadFrameAlpha = (float)PrevDataReadIndex - ReadFrameOffset;
+
+								if (SourceFileNumChannels == 1)
+								{
+									DecodedSampleDataPrev = DecodedData->PCMAudio[PrevDataReadIndex];
+									DecodedSampleDataNext = DecodedData->PCMAudio[NextDataReadIndex];
+								}
+								else
+								{
+									int32 SampleReadIndex = PrevDataReadIndex * SourceFileNumChannels;
+									float SampleA = DecodedData->PCMAudio[SampleReadIndex + ChannelAIndex];
+									float SampleB = DecodedData->PCMAudio[SampleReadIndex + ChannelBIndex];
+
+									DecodedSampleDataPrev = FMath::Lerp(SampleA, SampleB, ChannelAlpha);
+
+									SampleReadIndex = NextDataReadIndex * SourceFileNumChannels;
+									SampleA = DecodedData->PCMAudio[SampleReadIndex + ChannelAIndex];
+									SampleB = DecodedData->PCMAudio[SampleReadIndex + ChannelBIndex];
+
+									DecodedSampleDataNext = FMath::Lerp(SampleA, SampleB, ChannelAlpha);
+								}
+							}
+						}
+
+						// Store it here for interpolating between buffers
+						Grain.PrevASampleValue = DecodedSampleDataPrev;
+
+						float SampleValue = FMath::Lerp(DecodedSampleDataPrev, DecodedSampleDataNext, ReadFrameAlpha);
+
+						SampleValue *= GrainAmplitude;
+
+						// Write the value to the output buffer
+						MonoGrainBuffer[GrainWriteIndex++] = SampleValue;
+
+						Grain.CurrentRenderedFramesCount++;
+						NumFramesToRenderForBlock--;
+
+						Grain.CurrentReadFrame += Grain.CurrentPlaybackRate;
+						bGrainFinished = (Grain.CurrentRenderedFramesCount >= Grain.GrainDurationFrames);
+
+						// We're reading backward in the audio file and we've read off the chunk
+						if (!bGrainFinished && Grain.CurrentReadFrame < (float)DecodedData->FrameStart)
+						{
+							break;
+						}
+					}
+				}
+
+				// We have our generated mono grain buffer, now write it to the outputs
+				// This will do the spatialization of the grain
+				TArrayView<float> MonoGrainBufferView = TArrayView<float>(MonoGrainBuffer);
+
+				// Apply the grain volume scale on the mono grain buffer before panning
+				if (Grain.PreviousVolume < 0.0f)
+				{
+					Grain.PreviousVolume = Grain.CurrentVolume;
+				}
+
+				if (FMath::IsNearlyEqual(Grain.PreviousVolume, Grain.CurrentVolume))
+				{
+					Audio::ArrayMultiplyByConstantInPlace(MonoGrainBufferView, Grain.CurrentVolume);
+				}
+				else
+				{
+					Audio::ArrayFade(MonoGrainBufferView, Grain.PreviousVolume, Grain.CurrentVolume);
+					Grain.PreviousVolume = Grain.CurrentVolume;
+				}
+
+				MixMonoGrainToOutput(MonoGrainBufferView, Grain);
+
+				// If the grain has finished, then remove it from the active grain list
+				if (bGrainFinished)
+				{
+					DecodedChunks[Grain.DecodedDataChunkIndex].NumGrainsUsingChunk--;
+					check(DecodedChunks[Grain.DecodedDataChunkIndex].NumGrainsUsingChunk >= 0);
+					ActiveGrains.RemoveAtSwap(GrainIndex, EAllowShrinking::No);
+				}
+			}
+		}
+
+		void Execute()
+		{
+			// Zero everytime.
+			OutputAudio->Zero();
+			
+			const TArray<int32>& SpawnTriggerFrames = SpawnGrainTrigger->GetTriggeredFrames();
+			const TArray<int32>& StopTriggerFrames = StopTrigger->GetTriggeredFrames();
+
+			int32 LastStartTriggerFrame = SpawnTriggerFrames.Num() > 0 ? SpawnTriggerFrames.Last() : INDEX_NONE;
+			int32 LastStopTriggerFrame = StopTriggerFrames.Num() > 0 ? StopTriggerFrames.Last() : INDEX_NONE;
+
+			// Check if we're stopping this block
+
+			// Initialize the granulator if we're not currently playing
+			if (!bPlaying && LastStartTriggerFrame >= 0)
+			{
+				// Initialize things if we didn't trigger stop this block or the last stop trigger is before the last start trigger
+				if (LastStopTriggerFrame == INDEX_NONE || LastStopTriggerFrame < LastStartTriggerFrame)
+				{
+					bPlaying = Initialize();
+				}
+			}
+	
+			// We're not playing so just early return
+			if (!bPlaying)
+			{
+				return;
+			}
+
+			// Check if we're spawning new grains this block
+			if (SpawnTriggerFrames.Num() > 0)
+			{
+				// Wrap the grain location param to be between 0.0 and 1.0
+				float WrappedLocation = FMath::Fmod(*GrainLocationX, 1.0f);
+				if (WrappedLocation < 0.0f)
+				{
+					WrappedLocation += 1.0f;
+				}				
+				float WrappedSeekTimeSeconds = WrappedLocation * SourceFileDurationSeconds;
+				
+				float ClampedGrainDurationSeconds = 0.001f * FMath::Clamp(*GrainDuration, 0.0f, 2000.0f);
+				int32 ClampedGrainDurationFrames = (int32) (ClampedGrainDurationSeconds * Settings.GetSampleRate());
+				ClampedGrainDurationFrames = FMath::Max(ClampedGrainDurationFrames, 50);
+				
+				float ClampedPlaybackRate = *GrainPlaybackRate;
+				if (FMath::IsNearlyZero(ClampedPlaybackRate))
+				{
+					ClampedPlaybackRate = KINDA_SMALL_NUMBER;
+				}
+
+				float ClampedVolume = FMath::Clamp(*GrainVolume, -4.0f, 4.0f);
+
+				float WrappedPan = FMath::Fmod(*GrainPan, 2.0f);
+				if (WrappedPan < 0.0f)
+				{
+					WrappedPan += 2.0f;
+				}
+
+				float ClampedChannelY = FMath::Clamp(*GrainChannelY, 0.0f, 1.0f);
+
+				GranularEnvelopeType = OperatorData->GranularEnvelope;
+
+				FSpawnGrainData NewGrainSpawnData;
+				NewGrainSpawnData.LocationTimeSeconds = WrappedSeekTimeSeconds;
+				NewGrainSpawnData.LocationChannelFraction = ClampedChannelY;
+				NewGrainSpawnData.DurationFrames = ClampedGrainDurationFrames;
+				NewGrainSpawnData.PlaybackRate = ClampedPlaybackRate;
+				NewGrainSpawnData.Volume = ClampedVolume;
+				NewGrainSpawnData.Pan = WrappedPan;
+				NewGrainSpawnData.Envelope = GranularEnvelopeType;
+
+				// Only spawn new grains up until the first Stop-Trigger this frame
+				int32 FirstStopTrigger = StopTriggerFrames.Num() > 0 ? StopTriggerFrames[0] : Settings.GetNumFramesPerBlock();
+				for (int32 i = 0; i < SpawnTriggerFrames.Num(); ++i)
+				{
+					if (SpawnTriggerFrames[i] >= FirstStopTrigger)
+					{
+						break;
+					}
+
+					NewGrainSpawnData.StartFrameOffset = SpawnTriggerFrames[i];
+					SpawnGrain(NewGrainSpawnData);
+				}
+			}
+
+			// Zero the output of the channel agnostic audio
+			OutputAudio->Zero();
+
+			RenderActiveGrains(Settings.GetNumFramesPerBlock());
+		}
+
+		static FNodeClassMetadata GetNodeInfo()
+		{
+			using namespace GranulatorPrivate;
+			return FNodeClassMetadata
+			{
+				FNodeClassName{ "Experimental", "GranulatorOperator", "" },
+				1, // Major version
+				0, // Minor version
+				LOCTEXT("GranulatorNodeName", "Granulator"),	
+				LOCTEXT("GranulatorNodeNameDescription", "A Node that granulates an input sound wave asset."),	
+				TEXT("UE"), // Author
+				LOCTEXT("ExampleGranulatorPromptIfMissing", "Enable the MetaSoundExperimental Plugin"), // Prompt if missing
+				MakeInterface(TEXT("Cat:Stereo1Dot0")),
+				{}
+			};
+		}
+	private:
+		TSharedPtr<const FGranulatorOperatorData> OperatorData;
+		FTriggerReadRef SpawnGrainTrigger;
+		FTriggerReadRef StopTrigger;
+		FWaveAssetReadRef WaveAsset;
+		FFloatReadRef GrainLocationX;
+		FFloatReadRef GrainChannelY;
+		FFloatReadRef GrainDuration;
+		FFloatReadRef GrainPlaybackRate;
+		FFloatReadRef GrainVolume;
+		FFloatReadRef GrainPan;
+		FDiscreteChannelAgnosticTypeWriteRef OutputAudio;
+		TArray<float> MonoGrainBuffer;
+
+		FOperatorSettings Settings;
+		const FDiscreteChannelTypeFamily& DiscreteFamily;
+		
+		// Place to store decoded chunks. We need 2 chunks for actively playing grains plus a free chunk.
+		TArray<FDecodedDataChunk> DecodedChunks;
+
+		// The grain data used to render the granular audio
+		TArray<FGrain> ActiveGrains;
+
+		bool bPlaying = false;
+
+		// How much audio to decode per decode block
+		float DecodedAudioSizeInSeconds = 0.25f;
+
+		float SourceFileSampleRate = 0.0f;
+		float SourceFileDurationSeconds = 0.0f;
+		int32 SourceFileNumFrames = 0;
+		int32 SourceFileNumChannels = 1;
+
+		int32 GrainCount = 0;
+
+		FWaveAsset CurrentWaveAsset;
+		FSoundWaveProxyPtr SoundWaveProxyPtr;
+		TUniquePtr<FSoundWaveProxyReader> SoundWaveProxyReaderPtr;
+		TUniquePtr<Audio::ISampleRateConverter> SRC;
+		EGranularEnvelope GranularEnvelopeType = EGranularEnvelope::Hann;
+		TArray<Audio::Grain::FEnvelope> GrainEnvelopeData;
+		const Audio::IDiscretePanner* Panner = nullptr;
+	}; 
+		
+	using FGranulatorNode = TNodeFacade<FGranulatorOperator>;
+
+	// register node config.	
+	METASOUND_REGISTER_NODE_AND_CONFIGURATION(FGranulatorNode, FMetaSoundGranulatorNodeConfiguration);
+} // namespace Metasound
+
+TArray<FPropertyTextFName> UMetaSoundGranulatorNodeOptionsHelper::GetSoundFileFormatChannelOptions()
+{
+	// These are output from Unreal.
+	// Hardcoded for now.
+	static const TArray<FName> InternalFormats {
+		TEXT("Cat:Mono1Dot0"),
+		TEXT("Cat:Stereo2Dot0"),
+		TEXT("Cat:Surround5Dot1"),
+		TEXT("Cat:Surround7Dot1"),
+	};
+
+	const TArray<TSharedRef<const Audio::FChannelTypeFamily>> AllFormats = Audio::GetChannelRegistry().GetAllChannelFormats();
+	TArray<FPropertyTextFName> FormatsOptions;
+	for (const TSharedRef<const Audio::FChannelTypeFamily>& Format : AllFormats)
+	{
+		if (const Audio::FDiscreteChannelTypeFamily* Discrete = Format->Cast<Audio::FDiscreteChannelTypeFamily>())
+		{
+			if (InternalFormats.Contains(Discrete->GetName()))
+			{
+				FormatsOptions.Emplace(Format->GetName(), FText::FromString(Format->GetFriendlyName()));
+			}
+		}
+	}
+	return FormatsOptions;
+}
+
+TInstancedStruct<FMetasoundFrontendClassInterface> FMetaSoundGranulatorNodeConfiguration::OverrideDefaultInterface(const FMetasoundFrontendClass& InNodeClass) const
+{
+	return TInstancedStruct<FMetasoundFrontendClassInterface>::Make(FMetasoundFrontendClassInterface::GenerateClassInterface(Metasound::FGranulatorOperator::MakeInterface(this->OutputAudioTypeName)));
+}
+
+FMetaSoundGranulatorNodeConfiguration::FMetaSoundGranulatorNodeConfiguration()
+	: OperatorData(MakeShared<Metasound::GranulatorPrivate::FGranulatorOperatorData>(OutputAudioTypeName, GranularEnvelope))
+{
+}
+
+
+TSharedPtr<const Metasound::IOperatorData> FMetaSoundGranulatorNodeConfiguration::GetOperatorData() const
+{
+	OperatorData->OutputAudioTypeName = OutputAudioTypeName;
+	OperatorData->GranularEnvelope = GranularEnvelope;
+	return OperatorData;
+}
+
+#undef LOCTEXT_NAMESPACE // 
diff --git a/Engine/Plugins/Experimental/MetasoundExperimental/Source/MetasoundExperimentalEngineRuntime/Public/MetasoundGranulatorNode.h b/Engine/Plugins/Experimental/MetasoundExperimental/Source/MetasoundExperimentalEngineRuntime/Public/MetasoundGranulatorNode.h
new file mode 100644
--- /dev/null
+++ b/Engine/Plugins/Experimental/MetasoundExperimental/Source/MetasoundExperimentalEngineRuntime/Public/MetasoundGranulatorNode.h
@@ -0,0 +1,76 @@
+// Copyright Epic Games, Inc. All Rights Reserved.
+
+#pragma once
+
+#include "MetasoundFrontendDocument.h"
+#include "Templates/SharedPointer.h"
+#include "TypeFamily/ChannelTypeFamily.h"
+#include "UObject/PropertyText.h"
+
+#include "MetasoundGranulatorNode.generated.h"
+
+UCLASS()
+class UMetaSoundGranulatorNodeOptionsHelper : public UObject
+{
+	GENERATED_BODY()
+
+	UFUNCTION()
+	static TArray<FPropertyTextFName> GetSoundFileFormatChannelOptions();
+};
+
+
+UENUM()
+enum class EGranularEnvelope : uint8
+{
+	Gaussian = 0,
+	Triangle,
+	DownwardTriangle,
+	UpwardTriangle,
+	ExponentialDecay,
+	ExponentialAttack,
+	Hann,
+	Count UMETA(Hidden)
+};
+
+namespace Metasound::GranulatorPrivate
+{
+	class FGranulatorOperatorData final : public TOperatorData<FGranulatorOperatorData>
+	{
+	public:
+		static const FLazyName OperatorDataTypeName;
+
+		explicit FGranulatorOperatorData(const FName& InOutputAudioTypeName, const EGranularEnvelope& InGranularEnvelope)
+			: OutputAudioTypeName(InOutputAudioTypeName)
+			, GranularEnvelope(InGranularEnvelope)
+		{
+		}
+
+		FName OutputAudioTypeName;
+		EGranularEnvelope GranularEnvelope;
+	};
+}
+
+USTRUCT()
+struct FMetaSoundGranulatorNodeConfiguration : public FMetaSoundFrontendNodeConfiguration
+{
+	GENERATED_BODY()
+
+	FMetaSoundGranulatorNodeConfiguration();
+
+	// What the CAT output type is. Only azimuthal discrete channel types supported.
+	UPROPERTY(EditAnywhere, Category = General, meta = (GetOptions="MetaSoundGranulatorNodeOptionsHelper.GetSoundFileFormatChannelOptions"))
+	FName OutputAudioTypeName = TEXT("Cat:Stereo2Dot0");
+
+	// What grain envelope to use for the granulator
+	UPROPERTY(EditAnywhere, Category = General)
+	EGranularEnvelope GranularEnvelope = EGranularEnvelope::Hann;
+
+	/* Get the current interface for the class based upon the node extension */
+	virtual TInstancedStruct<FMetasoundFrontendClassInterface> OverrideDefaultInterface(const FMetasoundFrontendClass& InNodeClass) const override;
+
+	/* Pass data down to the operator. */
+	virtual TSharedPtr<const Metasound::IOperatorData> GetOperatorData() const override;
+
+private:
+	mutable TSharedPtr<Metasound::GranulatorPrivate::FGranulatorOperatorData> OperatorData;
+};
diff --git a/Engine/Plugins/Experimental/MetasoundExperimental/Source/MetasoundExperimentalRuntime/MetasoundExperimentalRuntime.Build.cs b/Engine/Plugins/Experimental/MetasoundExperimental/Source/MetasoundExperimentalRuntime/MetasoundExperimentalRuntime.Build.cs
--- a/Engine/Plugins/Experimental/MetasoundExperimental/Source/MetasoundExperimentalRuntime/MetasoundExperimentalRuntime.Build.cs
+++ b/Engine/Plugins/Experimental/MetasoundExperimental/Source/MetasoundExperimentalRuntime/MetasoundExperimentalRuntime.Build.cs
@@ -39,15 +39,20 @@ public class MetasoundExperimentalRuntime : ModuleRules
 				"MetasoundGraphCore",
 				"MetasoundFrontend",
 				"MetasoundStandardNodes",
-				"SignalProcessing"
+				"SignalProcessing",
+				"Serialization",
+				"Projects"
 			}
 		);
-
 		PrivateDefinitions.AddRange(
 			new string[]
 			{
 				"METASOUND_PLUGIN=MetasoundExperimental",
 				"METASOUND_MODULE=MetasoundExperimentalRuntime"
 			});
+		
+		RuntimeDependencies.Add("$(PluginDir)/Content/CatSchemas/Composite/*.json", StagedFileType.UFS);
+		RuntimeDependencies.Add("$(PluginDir)/Content/CatSchemas/Discrete/*.json", StagedFileType.UFS);
+		RuntimeDependencies.Add("$(PluginDir)/Content/CatSchemas/Soundfield/*.json", StagedFileType.UFS);
 	}
 }
diff --git a/Engine/Plugins/Experimental/MetasoundExperimental/Source/MetasoundExperimentalRuntime/Private/MetasoundCatBreakNode.cpp b/Engine/Plugins/Experimental/MetasoundExperimental/Source/MetasoundExperimentalRuntime/Private/MetasoundCatBreakNode.cpp
--- a/Engine/Plugins/Experimental/MetasoundExperimental/Source/MetasoundExperimentalRuntime/Private/MetasoundCatBreakNode.cpp
+++ b/Engine/Plugins/Experimental/MetasoundExperimental/Source/MetasoundExperimentalRuntime/Private/MetasoundCatBreakNode.cpp
@@ -20,27 +20,42 @@ namespace Metasound
 
 	namespace CatBreakNodePrivate
 	{
-		METASOUND_PARAM(InputFromCat,			"Input", "CAT to Cast");
-		METASOUND_PARAM(OutputToCat,			"Output", "CAT Result");
+		METASOUND_PARAM(InputFromCat,			"Input", "Cat to Break");
 
-		const FLazyName InputBaseName{"In"};
-
-		FName MakeOutputVertexName(const int32 InIndex)
+		static FName GetOutputPinName(const int32 InIndex, const Audio::FChannelTypeFamily& InFormat)
 		{
-			FName Name = InputBaseName;
-			Name.SetNumber(InIndex);
-			return Name;
+			if (TOptional<Audio::FChannelTypeFamily::FChannelName> Names = InFormat.GetChannelName(InIndex); Names.IsSet())
+			{
+				return Names->Name;
+			}
+			const FName Output(TEXT("Output"), InIndex);
+			return Output;
 		}
 
-		TOutputDataVertex<FAudioBuffer> MakeOutputDataVertex(const int32 InIndex, const FName InName, const FString& InFriendlyName)
-		{
-			const FName Name = MakeOutputVertexName(InIndex);
 
-			const FText DisplayName = METASOUND_LOCTEXT_FORMAT("In_DisplayName", "{0}", FText::FromString(InName.ToString()));
+		FOutputDataVertex MakeOutputDataVertex(const FName InName, const FString& InFriendlyName)
+		{
+			const FText DisplayName = METASOUND_LOCTEXT_FORMAT(
+				"DisplayName_{0}", "{0}", FText::FromName(InName));
 
-			return TOutputDataVertex<FAudioBuffer>{Name, FDataVertexMetadata{FText::FromString(*InFriendlyName), DisplayName}};
+			const FText Tooltip = METASOUND_LOCTEXT_FORMAT(
+				"Tooltip_{0}", "Channel {0}", FText::FromString(InFriendlyName));
+			
+			return FOutputDataVertex
+			{
+				InName,
+				GetMetasoundDataTypeName<FAudioBuffer>(),
+				FDataVertexMetadata
+				{
+					.Description = Tooltip,
+					.DisplayName = DisplayName,
+					.bIsAdvancedDisplay = false
+				},
+				EVertexAccessType::Reference
+			};
 		}
 		
+		
 		FVertexInterface MakeClassInterface(const FName Format)
 		{
 			const Audio::FChannelTypeFamily* Found = Audio::GetChannelRegistry().FindConcreteChannel(Format);
@@ -51,18 +66,19 @@ namespace Metasound
 			const int32 NumChannels = Found->NumChannels();
 
 			FInputVertexInterface InputInterface;
-            InputInterface.Add(TInputDataVertex<FChannelAgnosticType>(METASOUND_GET_PARAM_NAME_AND_METADATA(OutputToCat)));
+            InputInterface.Add(FInputDataVertex(METASOUND_GET_PARAM_NAME(InputFromCat), Format, METASOUND_GET_PARAM_METADATA(InputFromCat), EVertexAccessType::Reference));
 			
 			FOutputVertexInterface OutputInterface;
 			for (int32 ChannelIndex = 0; ChannelIndex < NumChannels; ++ChannelIndex )
 			{
 				if (TOptional<Audio::FChannelTypeFamily::FChannelName> Name = Found->GetChannelName(ChannelIndex); Name.IsSet())
 				{
-					OutputInterface.Add(MakeOutputDataVertex(ChannelIndex, Name->Name, Name->FriendlyName));
+					OutputInterface.Add(MakeOutputDataVertex(Name->Name, Name->FriendlyName));
 				}
 				else
 				{
-					OutputInterface.Add(MakeOutputDataVertex(ChannelIndex, TEXT("Output"), FString::FromInt(ChannelIndex) ));	
+					const FName ChannelName(TEXT("Output"), ChannelIndex);
+					OutputInterface.Add(MakeOutputDataVertex(ChannelName, FString::Printf(TEXT("Channel %d"), ChannelIndex)));	
 				}
 			}
 			return{
@@ -129,26 +145,7 @@ namespace Metasound
 			, Format(InFormat)
 		{}
 		virtual ~FCatBreakOperator() override = default;
-
-		static const FVertexInterface& GetDefaultInterface()
-		{
-			using namespace CatBreakNodePrivate ;
-			auto CreateDefaultInterface = []()-> FVertexInterface
-			{
-				// inputs
-				FInputVertexInterface InputInterface;
-				InputInterface.Add(TInputDataVertex<FChannelAgnosticType>(METASOUND_GET_PARAM_NAME_AND_METADATA(InputFromCat)));
-				
-				// outputs
-				FOutputVertexInterface OutputInterface;
-				
-				return FVertexInterface(MoveTemp(InputInterface), MoveTemp(OutputInterface));
-			}; // end lambda: CreateDefaultInterface()
-
-			static const FVertexInterface DefaultInterface = CreateDefaultInterface();
-			return DefaultInterface;
-		}
-			
+		
 		static TUniquePtr<IOperator> CreateOperator(const FBuildOperatorParams& InParams, FBuildResults& OutResults)
 		{
 			using namespace CatBreakNodePrivate;
@@ -170,7 +167,7 @@ namespace Metasound
 
 			// Create the input pin.
 			TDataReadReference<FChannelAgnosticType> InputPin = InputData.GetOrCreateDefaultDataReadReference<FChannelAgnosticType>(
-				METASOUND_GET_PARAM_NAME(OutputToCat), InParams.OperatorSettings);
+				METASOUND_GET_PARAM_NAME(InputFromCat), InParams.OperatorSettings);
 
 			// Create output based on format.
 			TArray<TDataWriteReference<FAudioBuffer>> OutputAudioVertices;
@@ -200,14 +197,14 @@ namespace Metasound
 		virtual void BindInputs(FInputVertexInterfaceData& InOutVertexData) override
 		{
 			using namespace CatBreakNodePrivate;
-			InOutVertexData.BindReadVertex(METASOUND_GET_PARAM_NAME(OutputToCat),InputCat);
+			InOutVertexData.BindReadVertex(METASOUND_GET_PARAM_NAME(InputFromCat),InputCat);
 		}
 		virtual void BindOutputs(FOutputVertexInterfaceData& InOutVertexData) override
 		{
 			using namespace CatBreakNodePrivate;
 			for (int32 i = 0; i < OutputAudioVertices.Num(); i++)
 			{
-				InOutVertexData.BindWriteVertex(MakeOutputVertexName(i), OutputAudioVertices[i]);
+				InOutVertexData.BindWriteVertex(GetOutputPinName(i, InputCat->GetType()), OutputAudioVertices[i]);
 			}
 		}
 
@@ -240,7 +237,7 @@ namespace Metasound
 				METASOUND_LOCTEXT("CatBreakNodeNameDescription", "A Node that Breaks CATs"),	
 				TEXT("UE"), // Author
 				METASOUND_LOCTEXT("ExampleConfigurablePromptIfMissing", "Enable the MetaSoundExperimental Plugin"), // Prompt if missing
-				GetDefaultInterface(),
+				MakeClassInterface(TEXT("Cat:Stereo2Dot0")),
 				{}
 			};
 		}
diff --git a/Engine/Plugins/Experimental/MetasoundExperimental/Source/MetasoundExperimentalRuntime/Private/MetasoundCatBreakNode.h b/Engine/Plugins/Experimental/MetasoundExperimental/Source/MetasoundExperimentalRuntime/Private/MetasoundCatBreakNode.h
--- a/Engine/Plugins/Experimental/MetasoundExperimental/Source/MetasoundExperimentalRuntime/Private/MetasoundCatBreakNode.h
+++ b/Engine/Plugins/Experimental/MetasoundExperimental/Source/MetasoundExperimentalRuntime/Private/MetasoundCatBreakNode.h
@@ -15,8 +15,8 @@ struct FMetaSoundCatBreakNodeConfiguration : public FMetaSoundFrontendNodeConfig
 
 	FMetaSoundCatBreakNodeConfiguration() = default;
 
-	UPROPERTY(EditAnywhere, Category = General, meta = (GetOptions="MetasoundCatCastingOptionsHelper.GetCastingOptions"))
-	FName Format = TEXT("Stereo");
+	UPROPERTY(EditAnywhere, Category = General, meta = (GetOptions="MetasoundCatCastingOptionsHelper.GetCastingOptions_NoAbstract"))
+	FName Format = TEXT("Cat:Stereo2Dot0");
 
 	UPROPERTY(EditAnywhere, Category = General)
 	EMetasoundCatCastingMethod TranscodeMethod = EMetasoundCatCastingMethod::ChannelDrop;
diff --git a/Engine/Plugins/Experimental/MetasoundExperimental/Source/MetasoundExperimentalRuntime/Private/MetasoundCatCastingNode.cpp b/Engine/Plugins/Experimental/MetasoundExperimental/Source/MetasoundExperimentalRuntime/Private/MetasoundCatCastingNode.cpp
--- a/Engine/Plugins/Experimental/MetasoundExperimental/Source/MetasoundExperimentalRuntime/Private/MetasoundCatCastingNode.cpp
+++ b/Engine/Plugins/Experimental/MetasoundExperimental/Source/MetasoundExperimentalRuntime/Private/MetasoundCatCastingNode.cpp
@@ -2,11 +2,13 @@
 
 #include "MetasoundCatCastingNode.h"
 
+#include "MetasoundCatDataSchemas.h"
 #include "MetasoundChannelAgnosticType.h"
 #include "MetasoundExecutableOperator.h"
 #include "MetasoundFacade.h"
 #include "MetasoundNodeRegistrationMacro.h"
 #include "MetasoundVertex.h"
+#include "DSP/MultiMono.h"
 #include "TypeFamily/ChannelTypeFamily.h"
 
 #define LOCTEXT_NAMESPACE "MetasoundStandardNodes_CatCastingNode"
@@ -49,144 +51,166 @@ namespace Metasound
 		// Linkage.
 		const FLazyName FCatCastingOperatorData::OperatorDataTypeName = TEXT("FCatCastingOperatorData");
 	}
-	
-	class FCatCastingOperator final : public TExecutableOperator<FCatCastingOperator>
+
+	static TDataWriteReference<FChannelAgnosticType>
+	MakePolyCat(const FName InDerivedType, const FOperatorSettings& InOpSettings)
 	{
-	public:
-		FCatCastingOperator(const FBuildOperatorParams& InParams, FChannelAgnosticTypeReadRef&& InInputCat, const CatCastingPrivate::FCatCastingOperatorData& InData, const FName InConcreteName)
-			: InputFrom(MoveTemp(InInputCat))
-			, ToFormatName(InConcreteName)
-			, OutputCastResult(FChannelAgnosticTypeWriteRef::CreateNew(InParams.OperatorSettings, InConcreteName))
-			, Settings(InParams.OperatorSettings)
-			, TranscodeMethod(InData.GetTranscodeMethod())
-			, MixMethod(InData.GetMixMethod())
-		{}
-		virtual ~FCatCastingOperator() override = default;
-
-		static const FVertexInterface& GetDefaultInterface()
-		{
-			using namespace CatCastingPrivate;
-			auto CreateDefaultInterface = []()-> FVertexInterface
-			{
-				// inputs
-				FInputVertexInterface InputInterface;
-				InputInterface.Add(TInputDataVertex<FChannelAgnosticType>(METASOUND_GET_PARAM_NAME_AND_METADATA(InputFromCat)));
-				
-				// outputs
-				FOutputVertexInterface OutputInterface;
-				OutputInterface.Add(TOutputDataVertex<FChannelAgnosticType>(METASOUND_GET_PARAM_NAME_AND_METADATA(OutputToCat)));
-				
-				return FVertexInterface(InputInterface, OutputInterface);
-			}; // end lambda: CreateDefaultInterface()
-
-			static const FVertexInterface DefaultInterface = CreateDefaultInterface();
-			return DefaultInterface;
-		}
-			
-		static TUniquePtr<IOperator> CreateOperator(const FBuildOperatorParams& InParams, FBuildResults& OutResults)
-		{
-			using namespace CatCastingPrivate;
+		using namespace Frontend;
+		const FLiteral Literal(InDerivedType.ToString());
+		TOptional<FAnyDataReference> Ref = IDataTypeRegistry::Get().CreateDataReference(InDerivedType, EDataReferenceAccessType::Write, Literal, InOpSettings);
+		return Ref->GetAs<TDataWriteReference<FChannelAgnosticType>>();
+	}
 
-			const FCatCastingOperatorData* CatTestingConfigData = CastOperatorData<const FCatCastingOperatorData>(InParams.Node.GetOperatorData().Get());
-			if (!CatTestingConfigData)
-			{
-				return MakeUnique<FNoOpOperator>();
-			}
-			const FName CastToName = CatTestingConfigData->GetToType();
-			
-			const FInputVertexInterfaceData& InputData = InParams.InputData;
-			const Audio::FChannelTypeFamily* ConcreteToType = Audio::GetChannelRegistry().FindConcreteChannel(CastToName);
-
-			TDataReadReference<FChannelAgnosticType> InputCat = InputData.GetOrCreateDefaultDataReadReference<FChannelAgnosticType>(
-				METASOUND_GET_PARAM_NAME(InputFromCat), InParams.OperatorSettings);
-
-			// Make sure the cast is to something sane, otherwise use the inputs type... 
-			const FName CastToNameSane = ConcreteToType ? ConcreteToType->GetName() : InputCat->GetTypeName();
-			
-			return MakeUnique<FCatCastingOperator>(
-				InParams,
-				MoveTemp(InputCat),
-				*CatTestingConfigData,
-				CastToNameSane
-			);
-		}
+	FCatCastingOperator::FCatCastingOperator(const FBuildOperatorParams& InParams, FChannelAgnosticTypeReadRef&& InInputCat,
+	                                         const CatCastingPrivate::FCatCastingOperatorData& InData, const FName InConcreteName)
+		: InputFrom(MoveTemp(InInputCat))
+		, ToFormatName(InConcreteName)
+		, OutputCastResult(MakePolyCat(InConcreteName, InParams.OperatorSettings))
+		, Settings(InParams.OperatorSettings)
+		, TranscodeMethod(InData.GetTranscodeMethod())
+		, MixMethod(InData.GetMixMethod())
+	{
+	}
 
-		virtual void BindInputs(FInputVertexInterfaceData& InOutVertexData) override
-		{
-			using namespace CatCastingPrivate;
-			InOutVertexData.BindReadVertex(METASOUND_GET_PARAM_NAME(InputFromCat) , InputFrom);
-		}
+	FVertexInterface FCatCastingOperator::GetInterface(const FName InSrc, const FName InDst) 
+	{
+		using namespace CatCastingPrivate;
+
+		// inputs
+		FInputVertexInterface InputInterface;
+		InputInterface.Add(
+			FInputDataVertex(
+				METASOUND_GET_PARAM_NAME(InputFromCat),
+				InSrc,
+				METASOUND_GET_PARAM_METADATA(InputFromCat)
+			)
+		);
 		
-		virtual void BindOutputs(FOutputVertexInterfaceData& InOutVertexData) override
-		{
-			using namespace CatCastingPrivate;
-			InOutVertexData.BindWriteVertex(METASOUND_GET_PARAM_NAME(OutputToCat),OutputCastResult);
-
-			// Create transcoder 
-			Transcoder = InputFrom->GetType().GetTranscoder(
-				{
-					.ToType = OutputCastResult->GetType(),
-					.TranscodeMethod = TranscodeMethod,
-					.MixMethod = MixMethod,
-				}
-			);
-		}
+		// outputs
+		FOutputVertexInterface OutputInterface;
+		OutputInterface.Add(
+			FOutputDataVertex(
+				METASOUND_GET_PARAM_NAME(OutputToCat),
+				InDst,
+				METASOUND_GET_PARAM_METADATA(OutputToCat),
+				EVertexAccessType::Reference
+			)
+		);
+		
+		return FVertexInterface(InputInterface, OutputInterface);
+	}
+
+	TUniquePtr<IOperator> FCatCastingOperator::CreateOperator(const FBuildOperatorParams& InParams, FBuildResults& OutResults)
+	{
+		using namespace CatCastingPrivate;
 
-		void Reset(const FResetParams& InParams)
+		const FCatCastingOperatorData* CatTestingConfigData = CastOperatorData<const FCatCastingOperatorData>(InParams.Node.GetOperatorData().Get());
+		if (!CatTestingConfigData)
 		{
-			Execute();
+			return MakeUnique<FNoOpOperator>();
 		}
+		const FName CastToName = CatTestingConfigData->GetToType();
 
-		void Execute()
-		{
-			if (Transcoder)
+		const FInputVertexInterfaceData& InputData = InParams.InputData;
+		const Audio::FChannelTypeFamily* ConcreteToType = Audio::GetChannelRegistry().FindConcreteChannel(CastToName);
+
+		TDataReadReference<FChannelAgnosticType> InputCat = InputData.GetOrCreateDefaultDataReadReference<FChannelAgnosticType>(
+			METASOUND_GET_PARAM_NAME(InputFromCat), InParams.OperatorSettings);
+
+		// Make sure the cast is to something sane, otherwise use the inputs type... 
+		const FName CastToNameSane = ConcreteToType ? ConcreteToType->GetName() : InputCat->GetTypeName();
+
+		return MakeUnique<FCatCastingOperator>(
+			InParams,
+			MoveTemp(InputCat),
+			*CatTestingConfigData,
+			CastToNameSane
+		);
+	}
+
+	void FCatCastingOperator::BindInputs(FInputVertexInterfaceData& InOutVertexData) 
+	{
+		using namespace CatCastingPrivate;
+		InOutVertexData.BindReadVertex(METASOUND_GET_PARAM_NAME(InputFromCat), InputFrom);
+	}
+
+	void FCatCastingOperator::BindOutputs(FOutputVertexInterfaceData& InOutVertexData) 
+	{
+		using namespace CatCastingPrivate;
+		InOutVertexData.BindWriteVertex(METASOUND_GET_PARAM_NAME(OutputToCat), OutputCastResult);
+
+		// Create transcoder 
+		Transcoder = InputFrom->GetType().GetTranscoder(
 			{
-				using namespace CatCastingPrivate;
-				using namespace Audio;
-				TStackArrayOfPointers<const float> Src = MakeMultiMonoPointersFromView(InputFrom->GetRawMultiMono(), Settings.GetNumFramesPerBlock(), InputFrom->NumChannels());
-				TStackArrayOfPointers<float> Dst = MakeMultiMonoPointersFromView(OutputCastResult->GetRawMultiMono(), Settings.GetNumFramesPerBlock(), OutputCastResult->NumChannels());
-				Transcoder(Src, Dst,  Settings.GetNumFramesPerBlock());
+				.ToType = OutputCastResult->GetType(),
+				.TranscodeMethod = TranscodeMethod,
+				.MixMethod = MixMethod,
 			}
-		}
-		static FNodeClassMetadata GetNodeInfo()
+		);
+	}
+
+	void FCatCastingOperator::Reset(const FResetParams& InParams)
+	{
+		Execute();
+	}
+
+	void FCatCastingOperator::Execute()
+	{
+		if (Transcoder)
 		{
 			using namespace CatCastingPrivate;
-			return FNodeClassMetadata
-			{
-				FNodeClassName{ "Experimental", "CatCastingOperator", "" },
-				1, // Major version
-				0, // Minor version
+			using namespace Audio;
+			TStackArrayOfPointers<const float> Src = MakeMultiMonoPointersFromView(InputFrom->GetRawMultiMono(), Settings.GetNumFramesPerBlock(),
+			                                                                       InputFrom->NumChannels());
+			TStackArrayOfPointers<float> Dst = MakeMultiMonoPointersFromView(OutputCastResult->GetRawMultiMono(), Settings.GetNumFramesPerBlock(),
+			                                                                 OutputCastResult->NumChannels());
+			Transcoder(Src, Dst, Settings.GetNumFramesPerBlock());
+		}
+	}
+
+	FNodeClassMetadata FCatCastingOperator::GetNodeInfo()
+	{
+		using namespace CatCastingPrivate;
+		static FName DefaultFormat = TEXT("Cat");
+		return FNodeClassMetadata
+		{
+			FNodeClassName{"Experimental", "CatCastingOperator", ""},
+			1, // Major version
+			0, // Minor version
 				METASOUND_LOCTEXT("CatCastingNodeName", "CAT Casting Node"),	
 				METASOUND_LOCTEXT("CatCastingNodeNameDescription", "A Node that allows Casting to CATs"),	
-				TEXT("UE"), // Author
+			TEXT("UE"), // Author
 				METASOUND_LOCTEXT("ExampleConfigurablePromptIfMissing", "Enable the MetaSoundExperimental Plugin"), // Prompt if missing
-				GetDefaultInterface(),
-				{}
-			};
-		}
-	private:
-		using FTranscoder = Audio::FChannelTypeFamily::FTranscoder;
-		FChannelAgnosticTypeReadRef InputFrom;
-		FName ToFormatName;
-		FChannelAgnosticTypeWriteRef OutputCastResult;
-		FOperatorSettings Settings;
-		FTranscoder Transcoder;
-		Audio::EChannelTranscodeMethod TranscodeMethod;
-		Audio::EChannelMapMonoUpmixMethod MixMethod;
-	}; // class FCatCastingOperator
-		
-	using FCatCastingNode = TNodeFacade<FCatCastingOperator>;
+			GetInterface(DefaultFormat, DefaultFormat),
+			{}
+		};
+	}
 
 	// register node config.	
 	METASOUND_REGISTER_NODE_AND_CONFIGURATION(FCatCastingNode, FMetaSoundCatCastingNodeConfiguration);
-} // namespace Metasound
+}; // class FCatCastingOperator
+
+
 
 
 TArray<FPropertyTextFName> UMetasoundCatCastingOptionsHelper::GetCastingOptions()
 {
-	const TArray<const Audio::FChannelTypeFamily*> AllFormats = Audio::GetChannelRegistry().GetAllChannelFormats();
+	TArray<TSharedRef<const Audio::FChannelTypeFamily>> AllFormats = Audio::GetChannelRegistry().GetAllChannelFormats();
+	TArray<FPropertyTextFName> FormatsOptions;
+	Algo::Transform(AllFormats, FormatsOptions, [](const TSharedRef<const Audio::FChannelTypeFamily>& i) -> FPropertyTextFName
+		{ return {.ValueString = i->GetName(), .DisplayName = FText::FromString(i->GetFriendlyName()) }; });
+	return FormatsOptions;
+}
+
+TArray<FPropertyTextFName> UMetasoundCatCastingOptionsHelper::GetCastingOptions_NoAbstract()
+{
+	TArray<TSharedRef<const Audio::FChannelTypeFamily>> AllFormats = Audio::GetChannelRegistry().GetAllChannelFormats();
 	TArray<FPropertyTextFName> FormatsOptions;
-	Algo::Transform(AllFormats, FormatsOptions, [](const Audio::FChannelTypeFamily* i) -> FPropertyTextFName
+	Algo::TransformIf(
+		AllFormats,
+		FormatsOptions,
+		[](const TSharedRef<const Audio::FChannelTypeFamily>& i) -> bool { return !i->IsAbstract(); },
+		[](const TSharedRef<const Audio::FChannelTypeFamily>& i) -> FPropertyTextFName
 		{ return {.ValueString = i->GetName(), .DisplayName = FText::FromString(i->GetFriendlyName()) }; });
 	return FormatsOptions;
 }
@@ -194,7 +218,10 @@ TArray<FPropertyTextFName> UMetasoundCatCastingOptionsHelper::GetCastingOptions(
 TInstancedStruct<FMetasoundFrontendClassInterface> FMetaSoundCatCastingNodeConfiguration::OverrideDefaultInterface(
 	const FMetasoundFrontendClass& InNodeClass) const
 {
-	return  TInstancedStruct<FMetasoundFrontendClassInterface>::Make(FMetasoundFrontendClassInterface::GenerateClassInterface(Metasound::FCatCastingOperator::GetDefaultInterface()));
+	static const FName DefaultInput = TEXT("Cat");
+	return TInstancedStruct<FMetasoundFrontendClassInterface>::Make(
+		FMetasoundFrontendClassInterface::GenerateClassInterface(
+			Metasound::FCatCastingOperator::GetInterface(DefaultInput,  Metasound::AddCatNamespaceToName(ToType))));
 }
 
 TSharedPtr<const Metasound::IOperatorData> FMetaSoundCatCastingNodeConfiguration::GetOperatorData() const
diff --git a/Engine/Plugins/Experimental/MetasoundExperimental/Source/MetasoundExperimentalRuntime/Private/MetasoundCatCastingNode.h b/Engine/Plugins/Experimental/MetasoundExperimental/Source/MetasoundExperimentalRuntime/Private/MetasoundCatCastingNode.h
--- a/Engine/Plugins/Experimental/MetasoundExperimental/Source/MetasoundExperimentalRuntime/Private/MetasoundCatCastingNode.h
+++ b/Engine/Plugins/Experimental/MetasoundExperimental/Source/MetasoundExperimentalRuntime/Private/MetasoundCatCastingNode.h
@@ -2,6 +2,8 @@
 
 #pragma once
 
+#include "MetasoundChannelAgnosticType.h"
+#include "MetasoundExecutableOperator.h"
 #include "MetasoundFrontendDocument.h"
 #include "Templates/SharedPointer.h"
 #include "TypeFamily/ChannelTypeFamily.h"
@@ -16,6 +18,9 @@ class UMetasoundCatCastingOptionsHelper : public UObject
 
 	UFUNCTION()
 	static TArray<FPropertyTextFName> GetCastingOptions();
+
+	UFUNCTION()
+	static TArray<FPropertyTextFName> GetCastingOptions_NoAbstract();
 };
 
 UENUM()
@@ -40,7 +45,7 @@ struct FMetaSoundCatCastingNodeConfiguration : public FMetaSoundFrontendNodeConf
 
 	FMetaSoundCatCastingNodeConfiguration() =  default;
 
-	UPROPERTY(EditAnywhere, Category = General, meta = (GetOptions="MetasoundCatCastingOptionsHelper.GetCastingOptions"))
+	UPROPERTY(EditAnywhere, Category = General, meta = (GetOptions="MetasoundCatCastingOptionsHelper.GetCastingOptions_NoAbstract"))
 	FName ToType = TEXT("Mono");
 	
 	UPROPERTY(EditAnywhere, Category = General)
@@ -55,3 +60,47 @@ struct FMetaSoundCatCastingNodeConfiguration : public FMetaSoundFrontendNodeConf
 	/* Pass data down to the operator. */
 	virtual TSharedPtr<const Metasound::IOperatorData> GetOperatorData() const override;
 };
+
+namespace Metasound
+{
+	namespace CatCastingPrivate
+	{
+		class FCatCastingOperatorData;
+	}
+
+	struct FBuildOperatorParams;
+
+	class FCatCastingOperator final : public TExecutableOperator<FCatCastingOperator>
+	{
+	public:
+		FCatCastingOperator(const FBuildOperatorParams& InParams, FChannelAgnosticTypeReadRef&& InInputCat, const CatCastingPrivate::FCatCastingOperatorData& InData, const FName InConcreteName);
+		virtual ~FCatCastingOperator() override = default;
+
+		static FVertexInterface GetInterface(const FName InSrc, const FName InDst);
+			
+		static TUniquePtr<IOperator> CreateOperator(const FBuildOperatorParams& InParams, FBuildResults& OutResults);
+
+		virtual void BindInputs(FInputVertexInterfaceData& InOutVertexData) override;
+		
+		virtual void BindOutputs(FOutputVertexInterfaceData& InOutVertexData) override;
+
+		void Reset(const FResetParams& InParams);
+
+		void Execute();
+	
+		static FNodeClassMetadata GetNodeInfo();
+	
+	private:
+		using FTranscoder = Audio::FChannelTypeFamily::FTranscoder;
+		FChannelAgnosticTypeReadRef InputFrom;
+		FName ToFormatName;
+		TDataWriteReference<FChannelAgnosticType> OutputCastResult;
+		FOperatorSettings Settings;
+		FTranscoder Transcoder;
+		Audio::EChannelTranscodeMethod TranscodeMethod;
+		Audio::EChannelMapMonoUpmixMethod MixMethod;
+	}; // class FCatCastingOperator
+
+	using FCatCastingNode = TNodeFacade<FCatCastingOperator>;
+
+}		
diff --git a/Engine/Plugins/Experimental/MetasoundExperimental/Source/MetasoundExperimentalRuntime/Private/MetasoundCatConversionNode.cpp b/Engine/Plugins/Experimental/MetasoundExperimental/Source/MetasoundExperimentalRuntime/Private/MetasoundCatConversionNode.cpp
new file mode 100644
--- /dev/null
+++ b/Engine/Plugins/Experimental/MetasoundExperimental/Source/MetasoundExperimentalRuntime/Private/MetasoundCatConversionNode.cpp
@@ -0,0 +1,116 @@
+﻿// Copyright Epic Games, Inc. All Rights Reserved.
+
+#include "MetasoundCatConversionNode.h"
+
+#include "MetasoundCatCastingNode.h"
+#include "MetasoundFrontendNodesCategories.h"
+#include "MetasoundNodeInterface.h"
+#include "MetasoundNodeRegistrationMacro.h"
+#include "MetasoundVertex.h"
+
+#define LOCTEXT_NAMESPACE "MetasoundFrontend"
+
+namespace Metasound
+{
+	struct FConvertDataTypeInfo
+	{
+		FName FromDataTypeName;
+		FText FromDataTypeText;
+		FName ToDataTypeName;
+		FText ToDataTypeText;
+	};
+	FVertexName GetInputName(const FConvertDataTypeInfo& InInfo)
+	{
+		return InInfo.FromDataTypeName;
+	}
+
+	FVertexName GetOutputName(const FConvertDataTypeInfo& InInfo)
+	{
+		return InInfo.ToDataTypeName;
+	}
+
+	FVertexInterface CreateVertexInterface(const FConvertDataTypeInfo& InInfo)
+	{
+		const FText InputDesc = {}; //METASOUND_LOCTEXT_FORMAT("AutoConvDisplayNamePatternFrom", "Input {0} value.", InInfo.FromDataTypeText);
+		const FText OutputDesc = {};//METASOUND_LOCTEXT_FORMAT("AutoConvDisplayNamePatternTo", "Output {0} value.", InInfo.ToDataTypeText);
+
+		return FVertexInterface(
+			FInputVertexInterface(
+				FInputDataVertex(GetInputName(InInfo), InInfo.FromDataTypeName, FDataVertexMetadata{ InputDesc }, EVertexAccessType::Reference)
+			),
+			FOutputVertexInterface(
+				FOutputDataVertex(GetOutputName(InInfo), InInfo.ToDataTypeName, FDataVertexMetadata{ OutputDesc }, EVertexAccessType::Reference)
+			)
+		);
+	}
+	FNodeClassMetadata CreateAutoConverterNodeMetadata(const FConvertDataTypeInfo& InInfo)
+	{
+		FNodeDisplayStyle DisplayStyle;
+		DisplayStyle.bShowName = false;
+		DisplayStyle.ImageName = TEXT("MetasoundEditor.Graph.Node.Conversion");
+		DisplayStyle.bShowInputNames = false;
+		DisplayStyle.bShowOutputNames = false;
+
+		const FText FromTypeText = InInfo.FromDataTypeText;
+		const FText ToTypeText = InInfo.ToDataTypeText;
+
+		FNodeClassMetadata Info;
+		Info.ClassName = {TEXT("Convert"), InInfo.ToDataTypeName, InInfo.FromDataTypeName};
+		Info.MajorVersion = 1;
+		Info.MinorVersion = 0;
+		/*
+		Info.DisplayName = METASOUND_LOCTEXT_FORMAT("Metasound_AutoConverterNodeDisplayNameFormat", "{0} to {1}", InInfo.FromDataTypeText,
+													InInfo.ToDataTypeText);
+		Info.Description = METASOUND_LOCTEXT_FORMAT("Metasound_AutoConverterNodeDescriptionNameFormat", "Converts from {0} to {1}.",
+													InInfo.FromDataTypeText, InInfo.ToDataTypeText);
+													*/
+		Info.Author = PluginAuthor;
+		Info.DisplayStyle = DisplayStyle;
+		Info.PromptIfMissing = PluginNodeMissingPrompt;
+		Info.DefaultInterface = CreateVertexInterface(InInfo);
+
+		Info.CategoryHierarchy.Emplace(NodeCategories::Conversions);
+
+		Info.Keywords =
+		{
+			METASOUND_LOCTEXT("MetasoundConvertKeyword", "Convert"),
+			InInfo.FromDataTypeText,
+			InInfo.ToDataTypeText
+		};
+
+		return Info;
+	}
+
+	void RegisterCatConvertNode()
+	{
+		FConvertDataTypeInfo DataTypeInfo
+		{
+			.FromDataTypeName = TEXT("Cat"),
+			.FromDataTypeText = FText::FromString(TEXT("Cat")),
+			.ToDataTypeName = TEXT("Cat:Discrete"),
+			.ToDataTypeText = FText::FromString(TEXT("Cat:Discrete")),
+		};
+
+		const FNodeClassMetadata& Metadata = Metasound::CreateAutoConverterNodeMetadata(DataTypeInfo);
+		const Frontend::FNodeRegistryKey Key(Metadata);
+
+		const Frontend::FConverterNodeClassInfo CatCast
+		{
+			// If this node has multiple input pins, we use this to designate which pin should be used.
+			.PreferredConverterInputPin = TEXT("Cat"),
+
+			// If this node has multiple output pins, we use this to designate which pin should be used.
+			.PreferredConverterOutputPin = TEXT("Cat"),
+
+			// The key for this node in the node registry.
+			.NodeKey = Key
+		};
+		Frontend::FConverterNodeClassRegistryKey ConverterNodeKey
+		{
+			.FromDataType = TEXT("Cat"),
+			.ToDataType = TEXT("Cat:Discrete"),
+		};
+
+		FMetasoundFrontendRegistryContainer::Get()->RegisterConversionNode(ConverterNodeKey, CatCast);
+	}
+}
diff --git a/Engine/Plugins/Experimental/MetasoundExperimental/Source/MetasoundExperimentalRuntime/Private/MetasoundCatConversionNode.h b/Engine/Plugins/Experimental/MetasoundExperimental/Source/MetasoundExperimentalRuntime/Private/MetasoundCatConversionNode.h
new file mode 100644
--- /dev/null
+++ b/Engine/Plugins/Experimental/MetasoundExperimental/Source/MetasoundExperimentalRuntime/Private/MetasoundCatConversionNode.h
@@ -0,0 +1,8 @@
+﻿// Copyright Epic Games, Inc. All Rights Reserved.
+
+#pragma once
+
+namespace Metasound
+{
+	void RegisterCatConvertNode();
+}
diff --git a/Engine/Plugins/Experimental/MetasoundExperimental/Source/MetasoundExperimentalRuntime/Private/MetasoundCatDataSchemas.cpp b/Engine/Plugins/Experimental/MetasoundExperimental/Source/MetasoundExperimentalRuntime/Private/MetasoundCatDataSchemas.cpp
new file mode 100644
--- /dev/null
+++ b/Engine/Plugins/Experimental/MetasoundExperimental/Source/MetasoundExperimentalRuntime/Private/MetasoundCatDataSchemas.cpp
@@ -0,0 +1,838 @@
+﻿// Copyright Epic Games, Inc. All Rights Reserved.
+
+#include "MetasoundCatDataSchemas.h"
+
+//#include "Backends/JsonStructDeserializerBackend.h"
+#include "MetasoundFrontend/Private/MetasoundJsonBackend.h"
+#include "HAL/FileManager.h"
+#include "HAL/IConsoleManager.h"
+#include "Interfaces/IPluginManager.h"
+#include "MetasoundCatMixerNode.h"
+#include "Misc/Paths.h"
+#include "Modules/ModuleManager.h"
+#include "Modules/ModuleManager.h"
+#include "StructDeserializer.h"
+#include "TypeFamily/ChannelTypeFamily.h"
+
+#define LOCTEXT_NAMESPACE "MetasoundFrontend"
+
+namespace Metasound
+{
+	static FAutoConsoleCommand DumpDataTypes(TEXT("au.metasound.dump_frontend_datatypes"), TEXT(""), FConsoleCommandDelegate::CreateLambda([]
+	{
+		using namespace Metasound::Frontend;
+		TArray<FName> Names;
+		IDataTypeRegistry::Get().GetRegisteredDataTypeNames(Names);
+		for (auto i : Names)
+		{
+			FDataTypeRegistryInfo Info;
+			IDataTypeRegistry::Get().GetDataTypeInfo(i, Info);
+			UE_LOG(LogMetaSound, Display, TEXT("GetDataTypeInfo: %s, bIsPolymorphic=%s, bIsAbstract=%s, Base=%s"),
+				*i.ToString(), ToCStr(LexToString(Info.bIsPolymorphic)),ToCStr(LexToString(Info.bIsAbstract)), *Info.ParentDataTypeName.ToString()  );
+		}
+	}));
+
+	static FAutoConsoleCommand DumpPolyTypes(TEXT("au.metasound.dump_poly_types"), TEXT(""), FConsoleCommandDelegate::CreateLambda([]
+	{
+		using namespace Metasound;
+	
+		TArray<FPolyTypeInfo> All;
+		FPolyRegistry::Get().GetAllRegisteredTypes(All);
+
+		for (FPolyTypeInfo& i : All)
+		{
+			UE_LOG(LogMetaSound, Display, TEXT("PolyType: %s, bIsPolymorphic=%s, Base=%s"),
+				*i.TypeName.ToString(), ToCStr(LexToString(i.bIsPolymorphic)), *i.BaseTypeName.ToString());
+		}
+	}));
+
+	// Brute this for now.
+	TOptional<ESpeakerShortNames> LexFromString(const FString& InName)
+	{
+		for (int32 i = static_cast<int32>(ESpeakerShortNames::DefaultChannel); i <= static_cast<int32>(ESpeakerShortNames::NumChannels); i++ )
+		{
+			const FString SpeakerName = LexToString(static_cast<ESpeakerShortNames>(i));
+			if (InName.Equals(SpeakerName, ESearchCase::IgnoreCase))
+			{
+				return static_cast<ESpeakerShortNames>(i);
+			}
+		}
+		// fail.
+		return {};
+	}
+
+	static FString MakePrettyString(const TArray<Audio::FDiscreteChannelTypeFamily::FSpeaker>& InSpeakers)
+	{
+		return FString::JoinBy(InSpeakers,TEXT(", "), [](const Audio::FDiscreteChannelTypeFamily::FSpeaker InSpeaker) -> FString
+		{
+			return FString::Printf(TEXT("[Name=%s,Az=%2.2f, El=%2.2f]"), LexToString(InSpeaker.Speaker), InSpeaker.AzimuthDegrees,
+				InSpeaker.ElevationDegrees);
+		});
+	}
+	
+
+	// Metasound Discrete, ingests the Schema.
+	class FMetasoundDiscreteChannelTypeFamily final : public Audio::FDiscreteChannelTypeFamily
+	{
+	public:
+		using Super = FDiscreteChannelTypeFamily;
+		static TArray<FSpeaker> MakeOrderFromSchema(const FCatDiscreteDataSchema& InSchema)
+		{
+			TArray<FSpeaker> Order;
+			for (auto i : InSchema.Order)
+			{
+				TOptional<ESpeakerShortNames> SpeakerId = LexFromString(i.ID.ToString());
+				ensure(SpeakerId.IsSet());
+				if (SpeakerId.IsSet())
+				{
+					Order.Emplace(
+						FSpeaker
+						{
+							.Speaker = *SpeakerId,
+							.AzimuthDegrees = i.AzimuthDegrees,
+							.ElevationDegrees = i.ElevationDegrees,
+							.bIsSpatialized = i.bIsSpatialized,
+						});
+				}
+			}
+			return Order;
+		}
+		
+		METASOUNDEXPERIMENTALRUNTIME_API FMetasoundDiscreteChannelTypeFamily(const FCatDiscreteDataSchema& InSchema, Audio::IChannelTypeRegistry& InRegistry)
+			: Super(InSchema.Name, InRegistry.FindChannel(InSchema.BaseType), *InSchema.FriendlyName, MakeOrderFromSchema(InSchema), InSchema.bIsParentsDefault, InSchema.bIsAbstract)
+		{
+			UE_LOG(LogMetaSound, Verbose, TEXT("Discrete: Unique=%s\tNumChannels=%d\tParent=%s\tFriendlyName=%s\tDefault=%s\tOrder=[%s]\tAbstract=%s"),
+				*InSchema.Name.ToString(),
+				Order.Num(),
+				*InSchema.BaseType.ToString(),
+				*InSchema.FriendlyName,
+				ToCStr(LexToString(InSchema.bIsParentsDefault)),
+				*MakePrettyString(Order),
+				ToCStr(LexToString(InSchema.bIsAbstract)));
+		}
+	};
+
+	class FMetasoundSoundfieldChannelTypeFamily final : public Audio::FSoundfieldChannelTypeFamily
+	{
+	public:
+		using Super = FSoundfieldChannelTypeFamily;
+		
+		METASOUNDEXPERIMENTALRUNTIME_API FMetasoundSoundfieldChannelTypeFamily(const FCatSoundfieldDataSchema& InSchema, Audio::IChannelTypeRegistry& InRegistry)
+			: Super(InSchema.Name,  InSchema.NumOrders, InRegistry.FindChannel(InSchema.BaseType), *InSchema.FriendlyName, InSchema.bIsParentsDefault, InSchema.bIsAbstract)
+		{
+			UE_LOG(LogMetaSound, Verbose, TEXT("Soundfield: Unique=%s\tParent=%s\tFriendlyName=%s\tOrder=%d\tDefault=%s\tAbstract=%s"),
+				*InSchema.Name.ToString(),
+				*InSchema.BaseType.ToString(),
+				*InSchema.FriendlyName,
+				GetAmbisonicsOrder(),
+				ToCStr(LexToString(InSchema.bIsParentsDefault)),
+				ToCStr(LexToString(InSchema.bIsAbstract)));
+		}	
+	};
+
+	class FMetasoundCompositeChannelTypeFamily final : public Audio::FCompositeChannelTypeFamily
+	{
+	public:
+		using Super = FCompositeChannelTypeFamily;
+
+		METASOUNDEXPERIMENTALRUNTIME_API FMetasoundCompositeChannelTypeFamily(const FCatCompositeDataSchema& InSchema, Audio::IChannelTypeRegistry& InRegistry)
+			: Super(InSchema.Name,InRegistry.FindChannel(InSchema.BaseType), *InSchema.FriendlyName, InSchema.bIsParentsDefault)
+		{
+			UE_LOG(LogMetaSound, Verbose, TEXT("Composite: Unique=%s\tParent=%s\tFriendlyName=%s\tDefault=%s\tAbstract=%s"),
+				*InSchema.Name.ToString(),
+				*InSchema.BaseType.ToString(),
+				*InSchema.FriendlyName,
+				ToCStr(LexToString(InSchema.bIsParentsDefault)),
+				ToCStr(LexToString(InSchema.bIsAbstract)));
+		}	
+	};
+
+	template <typename DataType, EVertexAccessType VertexAccess = EVertexAccessType::Reference>
+	class TPolyInputNodeOperatorFactory : public IOperatorFactory
+	{
+		static constexpr bool bIsReferenceVertexAccess = VertexAccess == EVertexAccessType::Reference;
+		static constexpr bool bIsValueVertexAccess = VertexAccess == EVertexAccessType::Value;
+
+		static_assert(bIsValueVertexAccess || bIsReferenceVertexAccess, "Unsupported EVertexAccessType");
+
+		// Choose which data reference type is created based on template parameters
+		using FPassThroughDataReference = std::conditional_t<bIsReferenceVertexAccess, TDataReadReference<DataType>, TDataValueReference<DataType>>;
+
+		// Return correct data reference type based on vertex access type for pass through scenario.
+		FPassThroughDataReference CreatePassThroughDataReference(const FAnyDataReference& InRef)
+		{
+			if constexpr (bIsReferenceVertexAccess)
+			{
+				return InRef.GetDataReadReference<DataType>();
+			}
+			else if constexpr (bIsValueVertexAccess)
+			{
+				return InRef.GetDataValueReference<DataType>();
+			}
+			else
+			{
+				static_assert("Unsupported EVertexAccessType");
+			}
+		}
+
+	public:
+		explicit TPolyInputNodeOperatorFactory(const FVertexName& InVertexName)
+			: VertexName(InVertexName)
+		{
+		}
+
+		virtual TUniquePtr<IOperator> CreateOperator(const FBuildOperatorParams& InParams, FBuildResults& OutResults) override
+		{
+			using namespace MetasoundInputNodePrivate;
+
+			if (const FAnyDataReference* Ref = InParams.InputData.FindDataReference(VertexName))
+			{
+				if constexpr (bIsReferenceVertexAccess)
+				{
+					if (EDataReferenceAccessType::Write == Ref->GetAccessType())
+					{
+						return MakeUnique<TNonOwningInputOperator<DataType, VertexAccess>>(VertexName, Ref->GetDataWriteReference<DataType>());
+					}
+				}
+				// Pass through input value
+				return MakeUnique<TPassThroughOperator<DataType, VertexAccess>>(VertexName, CreatePassThroughDataReference(*Ref));
+			}
+			else
+			{
+				// Owned input value
+				return MakeUnique<TInputOperator<DataType, VertexAccess>>(VertexName, InParams.OperatorSettings, InParams.InputData);
+			}
+		}
+
+	private:
+		FVertexName VertexName;
+	};
+	
+	// Polymorphic Specific INPUT NODE.
+	template<typename DataType, EVertexAccessType VertexAccess=EVertexAccessType::Reference>
+	class TPolyInputNode : public FInputNode
+	{
+		static constexpr bool bIsConstructorInput = VertexAccess == EVertexAccessType::Value;
+		static constexpr bool bIsSupportedConstructorInput = TIsConstructorVertexSupported<DataType>::Value && bIsConstructorInput;
+		static constexpr bool bIsReferenceInput = VertexAccess == EVertexAccessType::Reference;
+		static constexpr bool bIsSupportedReferenceInput = TLiteralTraits<DataType>::bIsParsableFromAnyLiteralType && bIsReferenceInput;
+
+		static constexpr bool bIsSupportedInput = bIsSupportedConstructorInput || bIsSupportedReferenceInput;
+
+	public:
+		// If true, this node can be instantiated by the Frontend.
+		static constexpr bool bCanRegister = bIsSupportedInput;
+		
+		static FNodeClassMetadata CreateNodeClassMetadata(const FVertexName& InVertexName, const FName InConcreteTypeName)
+		{
+			return FInputNode::CreateNodeClassMetadata(InVertexName, InConcreteTypeName, VertexAccess);
+		}
+
+			
+		explicit TPolyInputNode(FInputNodeConstructorParams&& InParams)
+		:	FInputNode(MoveTemp(InParams), InParams.PolymorphicTypeName, VertexAccess, MakeShared<TPolyInputNodeOperatorFactory<DataType, VertexAccess>>(InParams.VertexName))
+		{
+		}
+
+		explicit TPolyInputNode(const FVertexName& InVertexName, FNodeData InNodeData, TSharedRef<const FNodeClassMetadata> InClassMetadata)
+		: FInputNode(MakeShared<TPolyInputNodeOperatorFactory<DataType, VertexAccess>>(InVertexName), MoveTemp(InNodeData), MoveTemp(InClassMetadata))
+		{
+		}
+	};
+		
+	template<typename SchemaType, typename CppChannelType, typename CppDataType>
+	void AddCatDataFormat(const SchemaType& InSchema)
+	{
+		using namespace Metasound;
+		using namespace Metasound::Frontend;
+
+		using FBaseType = CppDataType;
+
+		struct FDataDrivenEntry : IDataTypeRegistryEntry 
+		{
+			static FDataTypeRegistryInfo MakeInfo(const SchemaType& InSchema)
+			{
+				FDataTypeRegistryInfo Info;
+
+				Info.DataTypeName =  InSchema.Name; // NamespaceThis(InSchema.Name);
+
+				// The preferred constructor argument type for creating instances of the data type.
+				Info.PreferredLiteralType = ELiteralType::String;
+
+				// Constructor argument support in TDataTypeLiteralFactory<TDataType>;
+				Info.bIsParsable = true;
+				Info.bIsDefaultParsable = false;
+				Info.bIsBoolParsable = false;
+				Info.bIsIntParsable = false;
+				Info.bIsFloatParsable = false;
+				Info.bIsStringParsable = true;
+				Info.bIsProxyParsable = false;
+				Info.bIsUniquePtrProxyParsable_DEPRECATED = false;
+				Info.bIsDefaultArrayParsable = false;
+				Info.bIsBoolArrayParsable = false;
+				Info.bIsIntArrayParsable = false;
+				Info.bIsFloatArrayParsable = false;
+				Info.bIsStringArrayParsable = false;
+				Info.bIsProxyArrayParsable = false;
+				Info.bIsUniquePtrProxyArrayParsable_DEPRECATED = false;
+
+				// Is a TEnum wrapped enum
+				Info.bIsEnum = false;
+
+				// Whether exact types are required
+				Info.bIsExplicit = false;
+
+				// Determines whether the type can be used with send/receive transmitters
+				Info.bIsTransmittable = false;
+
+				// Returns if DataType is a Variable type
+				Info.bIsVariable = false;
+
+				// Returns if DataType can be used for constructor vertices.
+				Info.bIsConstructorType = false;
+
+				// Returns if DataType represents an array type (ex. TArray<float>, TArray<int32>, etc.).
+				Info.bIsArrayType = false;
+
+				// Returns if DataType supports array parsing and passing array of base type to constructor.
+				Info.bIsArrayParseable = false;
+
+				// If provided in registration call, UClass this datatype was registered with.
+				Info.ProxyGeneratorClass = nullptr;
+			
+				// If this type is polymorphic.
+				Info.bIsPolymorphic = true;
+				Info.bIsAbstract = InSchema.bIsAbstract;
+
+				// If this type is polymorphic, what is its parent.
+				Info.ParentDataTypeName = InSchema.BaseType; // NamespaceThis(InSchema.BaseType);		
+			
+				return Info;
+			}
+
+			void InitNodeClasses()
+			{
+				const FName DataTypeName = Info.DataTypeName; 
+				const FName UnnamedVertex;
+
+				if (Info.bIsParsable)
+				{
+					const FNodeClassMetadata InputMeta  = FInputNode::CreateNodeClassMetadata(UnnamedVertex, DataTypeName, EVertexAccessType::Reference);
+					const FNodeClassMetadata OutputMeta = FOutputNode::CreateNodeClassMetadata(UnnamedVertex, DataTypeName, EVertexAccessType::Reference);
+					const FNodeClassMetadata LiteralMeta = FLiteralNode::GetNodeMetadata(DataTypeName, EVertexAccessType::Reference);
+
+					using namespace Metasound::VariableNames;
+
+					// Variable.
+					FNodeClassMetadata VarMeta = TVariableNode<FBaseType>::CreateNodeClassMetadata();                      
+					VarMeta.ClassName = GetVariableNodeClassName(DataTypeName);                                           
+
+					// Mutator.
+					FNodeClassMetadata MutMeta = MetasoundVariableNodesPrivate::TVariableMutatorOperator<FBaseType>::GetNodeInfo();
+					MutMeta.ClassName = GetVariableMutatorNodeClassName(DataTypeName);                                   
+
+					// Variable accessor
+					FNodeClassMetadata AccMeta = MetasoundVariableNodesPrivate::TVariableAccessorOperator<FBaseType>::GetNodeInfo();
+					AccMeta.ClassName = GetVariableAccessorNodeClassName(DataTypeName);                                 
+
+					// Deffered.
+					FNodeClassMetadata DefMeta = MetasoundVariableNodesPrivate::TVariableDeferredAccessorOperator<FBaseType>::GetNodeInfo();
+					DefMeta.ClassName = GetVariableDeferredAccessorNodeClassName(DataTypeName);
+					InitNodeClassesBase(InputMeta, OutputMeta, LiteralMeta, VarMeta, MutMeta, AccMeta, DefMeta);
+				}
+
+				if (Info.bIsConstructorType)
+				{
+					InitNodeClassesBaseConstructor(
+						FInputNode::CreateNodeClassMetadata(UnnamedVertex, DataTypeName, EVertexAccessType::Value),	
+						FOutputNode::CreateNodeClassMetadata(UnnamedVertex, DataTypeName, EVertexAccessType::Value)
+					);
+				}
+			}
+
+			void InitRawAssignmentFunction()
+			{
+				this->RawAssignmentFunction = [](const void* Src, void* Dest)
+					{
+						*(static_cast<FBaseType*>(Dest)) = *(static_cast<const FBaseType*>(Src));
+					};
+			}
+
+			void InitLiteralAssignmentFunction()
+			{
+				if  (Info.bIsParsable)
+				{
+					this->LiteralAssignmentFunction = [](const FOperatorSettings& InOperatorSettings, const FLiteral& InLiteral, const FAnyDataReference& OutDataRef)
+						{
+							*OutDataRef.GetWritableValue<FBaseType>() = TDataTypeLiteralFactory<FBaseType>::CreateExplicitArgs(InOperatorSettings, InLiteral);
+						};
+				}
+			}
+
+			void InitPolymorphicInterface(const SchemaType& InSchema)
+			{
+				this->PolymorphicInterface = MakeShared<FPolymorphicDataTypeBase>(InSchema.Name, InSchema.BaseType);
+			}
+	
+
+			FDataDrivenEntry(const SchemaType& InSchema)
+				: Info(MakeInfo(InSchema))
+			{
+				InitNodeClasses();
+				InitRawAssignmentFunction();
+				InitLiteralAssignmentFunction();
+				InitPolymorphicInterface(InSchema);
+			}
+
+			//FDataDrivenEntry(const FDataDrivenEntry&) = default;
+		
+			virtual ~FDataDrivenEntry() override = default;
+		
+			virtual TUniquePtr<IDataTypeRegistryEntry> Clone() const override
+			{
+				return {};//MakeUnique<FDataDrivenEntry>(*this);;
+			}
+		
+			virtual const FDataTypeRegistryInfo& GetDataTypeInfo() const override
+			{
+				return Info;
+			}
+			virtual TSharedPtr<const IEnumDataTypeInterface> GetEnumInterface() const override
+			{
+				return {}; 
+			}
+			virtual TSharedPtr<IPolymorphicDataTypeInterface> GetPolymorphicInterface() const override
+			{
+				return PolymorphicInterface;;
+			}
+		
+			virtual const FMetasoundFrontendClass& GetFrontendInputClass() const override
+			{
+				return InputClass;
+			}
+
+			virtual TSharedPtr<const Metasound::FNodeClassMetadata> GetInputClassMetadata() const override
+			{
+				return InputClassMetadata;
+			}
+
+			virtual const FMetasoundFrontendClass& GetFrontendConstructorInputClass() const override
+			{
+				return ConstructorInputClass;
+			}
+
+			virtual TSharedPtr<const Metasound::FNodeClassMetadata> GetConstructorInputClassMetadata() const override
+			{
+				return ConstructorInputClassMetadata;
+			}
+
+			virtual const FMetasoundFrontendClass& GetFrontendLiteralClass() const override
+			{
+				return LiteralClass;
+			}
+
+			virtual const FMetasoundFrontendClass& GetFrontendOutputClass() const override
+			{
+				return OutputClass;
+			}
+
+			virtual TSharedPtr<const FNodeClassMetadata> GetOutputClassMetadata() const override
+			{
+				return OutputClassMetadata;
+			}
+
+			virtual const FMetasoundFrontendClass& GetFrontendConstructorOutputClass() const override
+			{
+				return ConstructorOutputClass;
+			}
+
+			virtual TSharedPtr<const FNodeClassMetadata> GetConstructorOutputClassMetadata() const override
+			{
+				return ConstructorOutputClassMetadata;
+			}
+
+			virtual const FMetasoundFrontendClass& GetFrontendVariableClass() const override
+			{
+				return VariableClass;
+			}
+
+			virtual TSharedPtr<const FNodeClassMetadata> GetVariableClassMetadata() const override
+			{
+				return VariableClassMetadata;
+			}
+
+			virtual const FMetasoundFrontendClass& GetFrontendVariableMutatorClass() const override
+			{
+				return VariableMutatorClass;
+			}
+
+			virtual TSharedPtr<const FNodeClassMetadata> GetVariableMutatorClassMetadata() const override
+			{
+				return VariableMutatorClassMetadata;
+			}
+
+			virtual const FMetasoundFrontendClass& GetFrontendVariableAccessorClass() const override
+			{
+				return VariableAccessorClass;
+			}
+
+			virtual TSharedPtr<const FNodeClassMetadata> GetVariableAccessorClassMetadata() const override
+			{
+				return VariableAccessorClassMetadata;
+			}
+
+			virtual const FMetasoundFrontendClass& GetFrontendVariableDeferredAccessorClass() const override
+			{
+				return VariableDeferredAccessorClass;
+			}
+
+			virtual TSharedPtr<const FNodeClassMetadata> GetVariableDeferredAccessorClassMetadata() const override
+			{
+				return VariableDeferredAccessorClassMetadata;
+			}
+
+			virtual const IParameterAssignmentFunction& GetRawAssignmentFunction() const override
+			{
+				return RawAssignmentFunction;
+			}
+
+			virtual FLiteralAssignmentFunction GetLiteralAssignmentFunction() const override
+			{
+				return LiteralAssignmentFunction;
+			}
+
+			virtual TUniquePtr<INode> CreateOutputNode(FNodeData InNodeData) const override
+			{
+				if (OutputClassMetadata)
+				{
+					const FOutputVertexInterface& Outputs = InNodeData.Interface.GetOutputInterface();
+					if (ensure(Outputs.Num() == 1))
+					{
+						const FVertexName& VertexName = Outputs.At(0).VertexName;
+						return MakeUnique<FOutputNode>(VertexName, MoveTemp(InNodeData), OutputClassMetadata.ToSharedRef());
+					}
+				}
+
+				return TUniquePtr<INode>(nullptr);
+			}
+
+			virtual TUniquePtr<INode> CreateConstructorOutputNode(FNodeData InNodeData) const override
+			{
+				if (ConstructorOutputClassMetadata)
+				{
+					const FOutputVertexInterface& Outputs = InNodeData.Interface.GetOutputInterface();
+					if (ensure(Outputs.Num() == 1))
+					{
+						const FVertexName& VertexName = Outputs.At(0).VertexName;
+						return MakeUnique<FOutputNode>(VertexName, MoveTemp(InNodeData), ConstructorOutputClassMetadata.ToSharedRef());
+					}
+				}
+				return TUniquePtr<INode>(nullptr);
+			}
+
+			FORCENOINLINE void InitNodeClassesBase(
+				const FNodeClassMetadata& InInputClassMetadata,
+				const FNodeClassMetadata& InOutputClassMetadata,
+				const FNodeClassMetadata& InLiteralPrototypeMetadata,
+				const FNodeClassMetadata& InVariableClassMetadata,
+				const FNodeClassMetadata& InVariableMutatorClassMetadata,
+				const FNodeClassMetadata& InVariableAccessorClassMetadata,
+				const FNodeClassMetadata& InVariableDeferredAccessorClassMetadata
+			)
+			{
+				this->InputClassMetadata = MakeShared<FNodeClassMetadata>(InInputClassMetadata);
+				this->InputClass = GenerateClass(*this->InputClassMetadata, EMetasoundFrontendClassType::Input);
+
+				this->OutputClassMetadata = MakeShared<Metasound::FNodeClassMetadata>(InOutputClassMetadata);
+				this->OutputClass = GenerateClass(*this->OutputClassMetadata, EMetasoundFrontendClassType::Output);
+
+				this->LiteralClass = GenerateClass(InLiteralPrototypeMetadata, EMetasoundFrontendClassType::Literal);
+
+				this->VariableClassMetadata = MakeShared<FNodeClassMetadata>(InVariableClassMetadata);
+				this->VariableClass = GenerateClass(*this->VariableClassMetadata, EMetasoundFrontendClassType::Variable);
+
+				this->VariableMutatorClassMetadata = MakeShared<FNodeClassMetadata>(InVariableMutatorClassMetadata);
+				this->VariableMutatorClass = GenerateClass(*this->VariableMutatorClassMetadata, EMetasoundFrontendClassType::VariableMutator);
+
+				this->VariableAccessorClassMetadata = MakeShared<FNodeClassMetadata>(InVariableAccessorClassMetadata);
+				this->VariableAccessorClass = GenerateClass(*this->VariableAccessorClassMetadata, EMetasoundFrontendClassType::VariableAccessor);
+
+				this->VariableDeferredAccessorClassMetadata = MakeShared<FNodeClassMetadata>(InVariableDeferredAccessorClassMetadata);
+				this->VariableDeferredAccessorClass = GenerateClass(*this->VariableDeferredAccessorClassMetadata, EMetasoundFrontendClassType::VariableDeferredAccessor);
+			}
+
+			FORCENOINLINE void InitNodeClassesBaseConstructor(
+				const FNodeClassMetadata& InConstructorInputClassMetadata,
+				const FNodeClassMetadata& InConstructorOutputClassMetadata
+			)
+			{
+				this->ConstructorInputClassMetadata = MakeShared<Metasound::FNodeClassMetadata>(InConstructorInputClassMetadata);
+				this->ConstructorInputClass = Metasound::Frontend::GenerateClass(*this->ConstructorInputClassMetadata, EMetasoundFrontendClassType::Input);
+				this->ConstructorOutputClassMetadata = MakeShared<Metasound::FNodeClassMetadata>(InConstructorOutputClassMetadata);
+				this->ConstructorOutputClass = Metasound::Frontend::GenerateClass(*this->ConstructorOutputClassMetadata, EMetasoundFrontendClassType::Output);
+			}
+		
+		
+			virtual TUniquePtr<INode> CreateInputNode(Metasound::FNodeData InNodeData) const override
+			{
+				if (InputClassMetadata)
+				{
+					const FInputVertexInterface& Inputs = InNodeData.Interface.GetInputInterface();
+					if (ensure(Inputs.Num() == 1))
+					{
+						const FVertexName& VertexName = Inputs.At(0).VertexName;
+						return MakeUnique<TPolyInputNode<FBaseType, EVertexAccessType::Reference>>(VertexName, MoveTemp(InNodeData), InputClassMetadata.ToSharedRef());
+					}
+				}
+				return TUniquePtr<INode>(nullptr);
+			}
+
+			using TDataType = FBaseType;
+
+			virtual TUniquePtr<INode> CreateConstructorInputNode(FNodeData InNodeData) const override
+			{
+				if (Info.bIsParsable && Info.bIsConstructorType)
+				{
+					if (ConstructorInputClassMetadata)
+					{
+						const FInputVertexInterface& Inputs = InNodeData.Interface.GetInputInterface();
+						if (ensure(Inputs.Num() == 1))
+						{
+							const FVertexName& VertexName = Inputs.At(0).VertexName;
+							return MakeUnique<TPolyInputNode<TDataType, EVertexAccessType::Value>>(
+								VertexName, MoveTemp(InNodeData), ConstructorInputClassMetadata.ToSharedRef());
+						}
+					}
+				}
+				return TUniquePtr<INode>(nullptr);
+			}
+
+			virtual TUniquePtr<INode> CreateVariableNode(FLiteral InLiteral, FNodeData InNodeData) const override
+			{
+				if (Info.bIsParsable)
+				{
+					if (VariableClassMetadata)
+					{
+						return MakeUnique<TVariableNode<TDataType>>(MoveTemp(InLiteral), MoveTemp(InNodeData), VariableClassMetadata.ToSharedRef());
+					}
+				}
+				return TUniquePtr<INode>(nullptr);
+			}
+
+			virtual TUniquePtr<INode> CreateVariableMutatorNode(FNodeData InNodeData) const override
+			{
+				if (Info.bIsParsable)
+				{
+					if (VariableMutatorClassMetadata)
+					{
+						return MakeUnique<TVariableMutatorNode<TDataType>>(MoveTemp(InNodeData), VariableMutatorClassMetadata.ToSharedRef());
+					}
+				}
+				return TUniquePtr<INode>(nullptr);
+			}
+
+			virtual TUniquePtr<INode> CreateVariableAccessorNode(const FNodeInitData& InParams) const override
+			{
+				if (Info.bIsParsable)
+				{
+					return MakeUnique<TVariableAccessorNode<TDataType>>(InParams);
+				}
+				else
+				{
+					return TUniquePtr<INode>(nullptr);
+				}
+			}
+
+			virtual TUniquePtr<INode> CreateVariableAccessorNode(FNodeData InNodeData) const override
+			{
+				if (Info.bIsParsable)
+				{
+					if (VariableAccessorClassMetadata)
+					{
+						return MakeUnique<TVariableAccessorNode<TDataType>>(MoveTemp(InNodeData), VariableAccessorClassMetadata.ToSharedRef());
+					}
+				}
+				return TUniquePtr<INode>(nullptr);
+			}
+
+			virtual TUniquePtr<INode> CreateVariableDeferredAccessorNode(FNodeData InNodeData) const override
+			{
+				if (Info.bIsParsable)
+				{
+					if (VariableAccessorClassMetadata)
+					{
+						return MakeUnique<TVariableDeferredAccessorNode<TDataType>>(MoveTemp(InNodeData), VariableAccessorClassMetadata.ToSharedRef());
+					}
+				}
+				return TUniquePtr<INode>(nullptr);
+			}
+
+			virtual TOptional<FAnyDataReference> CreateDataReference(EDataReferenceAccessType InAccessType, const FLiteral& InLiteral,
+																	 const FOperatorSettings& InOperatorSettings) const override
+			{
+				using enum EDataReferenceAccessType;;
+				// Use the CPP types ID. For cat that will mean Derived,Composite,Soundfield.
+				auto ID = GetMetasoundDataTypeId<CppDataType>();
+				switch (InAccessType)
+				{
+				case Read: return  FAnyDataReference(TDataReadReference<FChannelAgnosticType>::CreateNewDerivedAs<CppDataType>(Info.DataTypeName, ID, InOperatorSettings, Info.DataTypeName));
+				case Write: return FAnyDataReference(TDataWriteReference<FChannelAgnosticType>::CreateNewDerivedAs<CppDataType>(Info.DataTypeName, ID, InOperatorSettings, Info.DataTypeName));
+				case Value: return FAnyDataReference(TDataValueReference<FChannelAgnosticType>::CreateNewDerivedAs<CppDataType>(Info.DataTypeName, ID, InOperatorSettings, Info.DataTypeName));
+				default:
+					checkNoEntry();
+				}
+				return {};
+			}
+
+
+			virtual TSharedPtr<Audio::IProxyData> CreateProxy(UObject* InObject) const override { return {}; }
+			virtual TSharedPtr<IDataChannel, ESPMode::ThreadSafe> CreateDataChannel(const Metasound::FOperatorSettings&) const override { return {}; }
+
+			// stubbed.
+			virtual TUniquePtr<INode> CreateVariableNode(Metasound::FVariableNodeConstructorParams&& InInitParams) const override
+			{
+				return {};
+			}
+
+			virtual TUniquePtr<INode> CreateLiteralNode(Metasound::FLiteralNodeConstructorParams&&) const override
+			{
+				return {};
+			}
+
+			virtual TUniquePtr<INode> CreateReceiveNode(const Metasound::FNodeInitData&) const override
+			{
+				return {};
+			}
+
+			virtual TUniquePtr<INode> CreateConstructorOutputNode(Metasound::FOutputNodeConstructorParams&&) const override
+			{
+				return {};
+			}
+
+			virtual TUniquePtr<INode> CreateOutputNode(Metasound::FOutputNodeConstructorParams&&) const override
+			{
+				return {};
+			}
+
+			virtual TUniquePtr<INode> CreateVariableMutatorNode(const Metasound::FNodeInitData&) const override
+			{
+				return {};
+			}
+
+			virtual TUniquePtr<INode> CreateVariableDeferredAccessorNode(const Metasound::FNodeInitData&) const override
+			{
+				return {};
+			}
+
+			virtual TUniquePtr<INode> CreateInputNode(Metasound::FInputNodeConstructorParams&&) const override
+			{
+				return {};
+			}
+
+			virtual TUniquePtr<INode> CreateConstructorInputNode(Metasound::FInputNodeConstructorParams&&) const override
+			{
+				return {};
+			}
+
+			FDataTypeRegistryInfo Info;
+			FMetasoundFrontendClass InputClass;
+			TSharedPtr<const FNodeClassMetadata> InputClassMetadata;
+			FMetasoundFrontendClass ConstructorInputClass;
+			TSharedPtr<const FNodeClassMetadata> ConstructorInputClassMetadata;
+			FMetasoundFrontendClass OutputClass;
+			TSharedPtr<const FNodeClassMetadata> OutputClassMetadata;
+			FMetasoundFrontendClass ConstructorOutputClass;
+			TSharedPtr<const FNodeClassMetadata> ConstructorOutputClassMetadata;
+
+			FMetasoundFrontendClass LiteralClass;
+
+			FMetasoundFrontendClass VariableClass;
+			TSharedPtr<const FNodeClassMetadata> VariableClassMetadata;
+			FMetasoundFrontendClass VariableMutatorClass;
+			TSharedPtr<const FNodeClassMetadata> VariableMutatorClassMetadata;
+			FMetasoundFrontendClass VariableAccessorClass;
+			TSharedPtr<const FNodeClassMetadata> VariableAccessorClassMetadata;
+			FMetasoundFrontendClass VariableDeferredAccessorClass;
+			TSharedPtr<const FNodeClassMetadata> VariableDeferredAccessorClassMetadata;
+			TSharedPtr<IEnumDataTypeInterface> EnumInterface;
+			TSharedPtr<IPolymorphicDataTypeInterface> PolymorphicInterface;
+			IParameterAssignmentFunction RawAssignmentFunction;
+			FLiteralAssignmentFunction LiteralAssignmentFunction = nullptr;
+		};
+
+		UE_LOG(LogTemp, Display, TEXT("Registering: %s"), *InSchema.Name.ToString())
+		ensure(Audio::GetChannelRegistry().RegisterType(InSchema.Name, MakeUnique<CppChannelType>(InSchema, Audio::GetChannelRegistry())));
+		ensure(IDataTypeRegistry::Get().RegisterDataType(MakeUnique<FDataDrivenEntry>(InSchema)));
+	}
+
+	template<typename SchemaType, typename CppChannelType, typename CppDataType>
+	void LoadSchemas(const FString& Path)
+	{
+		TArray<FString> SchemasFiles;
+		IFileManager::Get().FindFiles(SchemasFiles, *Path, TEXT(".json"));
+
+		TMap<FName, SchemaType> Schemas;
+		for (int32 i = 0; i < SchemasFiles.Num(); i++)
+		{
+			const FString File = Path / SchemasFiles[i];
+			if (TUniquePtr<FArchive> Reader = TUniquePtr<FArchive>(IFileManager::Get().CreateFileReader(*File)))
+			{
+				TJsonStructDeserializerBackend<DefaultCharType> Backend(*Reader);
+				SchemaType Schema;
+				FStructDeserializer::Deserialize(Schema, Backend);
+				ensureMsgf(Schema.FamilyType == CppChannelType::GetFamilyTypeName(),
+					TEXT("Failed load schema. File=%s, LastError=%s"), *File, *Backend.GetLastErrorMessage());
+				
+				Schema.Name = AddCatNamespaceToName(Schema.Name);
+				Schema.BaseType = AddCatNamespaceToName(Schema.BaseType);
+				ensure(Schemas.Find(Schema.Name) == nullptr);
+				Schemas.Add(Schema.Name, Schema);
+			}
+		}
+		
+		int32 Attempts = 0;
+		while (Schemas.Num() > 0)
+		{
+			if (Attempts++ > 100)
+			{
+				checkNoEntry();
+				break;
+			}
+			for (auto i : Schemas)
+			{
+				// Check base is registered.
+				if (Audio::GetChannelRegistry().FindChannel(i.Value.BaseType))
+				{
+					AddCatDataFormat<SchemaType,CppChannelType, CppDataType>(i.Value);
+					Schemas.Remove(i.Key);
+					break;
+				}
+			}
+		}
+	}
+
+	static FString GetPluginContentPath()
+	{
+		if (const TSharedPtr<IPlugin> Plugin = IPluginManager::Get().FindPlugin(TEXT("MetasoundExperimental")); Plugin.IsValid())
+		{
+			return Plugin->GetContentDir();
+		}
+
+		// Some targets don't appear to return the plugin, so fall-back to this.
+		const FString ContentPath = FPaths::Combine(FPaths::EnginePluginsDir(), TEXT("Experimental/MetasoundExperimental/Content"));;
+		return ContentPath;
+	}
+	
+	void LoadPolymorphicDataSchemas()
+	{
+		const FString PluginContentDir = GetPluginContentPath();
+		const FString SchemaPath = PluginContentDir / TEXT("CatSchemas");
+		LoadSchemas<FCatDiscreteDataSchema,		FMetasoundDiscreteChannelTypeFamily,	FDiscreteChannelAgnosticType>(		SchemaPath / TEXT("Discrete"));
+		LoadSchemas<FCatSoundfieldDataSchema,	FMetasoundSoundfieldChannelTypeFamily,	FSoundfieldChannelAgnosticType>(	SchemaPath / TEXT("Soundfield"));
+		LoadSchemas<FCatCompositeDataSchema,	FMetasoundCompositeChannelTypeFamily,	FCompositeChannelAgnosticType>(		SchemaPath / TEXT("Composite"));
+	}
+
+
+} 
+
+#undef LOCTEXT_NAMESPACE
diff --git a/Engine/Plugins/Experimental/MetasoundExperimental/Source/MetasoundExperimentalRuntime/Private/MetasoundCatDataSchemas.h b/Engine/Plugins/Experimental/MetasoundExperimental/Source/MetasoundExperimentalRuntime/Private/MetasoundCatDataSchemas.h
new file mode 100644
--- /dev/null
+++ b/Engine/Plugins/Experimental/MetasoundExperimental/Source/MetasoundExperimentalRuntime/Private/MetasoundCatDataSchemas.h
@@ -0,0 +1,25 @@
+﻿// Copyright Epic Games, Inc. All Rights Reserved.
+
+#pragma once
+
+#include "UObject/NameTypes.h"
+
+namespace Metasound
+{
+	void LoadPolymorphicDataSchemas();
+
+	inline FName AddNamespaceToName(const FName InDatatype, const FName InNamespace)
+	{
+		if (InNamespace != InDatatype && !InDatatype.ToString().StartsWith(InNamespace.ToString()))
+		{
+			return *FString::Printf(TEXT("%s%s"), *InNamespace.ToString(), *InDatatype.ToString());
+		}
+		return InDatatype;
+	}
+
+	inline FName AddCatNamespaceToName(const FName InDatatype)
+	{
+		static const FName Cat(TEXT("Cat:"));
+		return AddNamespaceToName(InDatatype, Cat);
+	}
+}
diff --git a/Engine/Plugins/Experimental/MetasoundExperimental/Source/MetasoundExperimentalRuntime/Private/MetasoundCatMakeNode.cpp b/Engine/Plugins/Experimental/MetasoundExperimental/Source/MetasoundExperimentalRuntime/Private/MetasoundCatMakeNode.cpp
--- a/Engine/Plugins/Experimental/MetasoundExperimental/Source/MetasoundExperimentalRuntime/Private/MetasoundCatMakeNode.cpp
+++ b/Engine/Plugins/Experimental/MetasoundExperimental/Source/MetasoundExperimentalRuntime/Private/MetasoundCatMakeNode.cpp
@@ -22,26 +22,28 @@ namespace Metasound
 		METASOUND_PARAM(InputFromCat,			"Input", "CAT to Cast");
 		METASOUND_PARAM(OutputToCat,			"Output", "CAT Result");
 
-		const FLazyName InputBaseName{"In"};
-
-		FName MakeInputVertexName(const int32 InIndex)
-		{
-			FName Name = InputBaseName;
-			Name.SetNumber(InIndex);
-			return Name;
-		}
-
-		FInputDataVertex MakeInputDataVertex(const int32 InIndex, const FName InName, const FString& InFriendlyName)
+		FInputDataVertex MakeInputDataVertex(const FName InName, const FString& InFriendlyName)
 		{
-			const FName InputName = MakeInputVertexName(InIndex);
-
-			const FText InputDisplayName = METASOUND_LOCTEXT_FORMAT("In_DisplayName", "{0}", FText::FromString(InName.ToString()));
+			const FText DisplayName = METASOUND_LOCTEXT_FORMAT(
+				"DisplayName_{0}", "{0}", FText::FromName(InName));
 
-			return TInputDataVertex<FAudioBuffer>{InputName, FDataVertexMetadata{FText::FromString(*InFriendlyName), InputDisplayName}};
+			const FText Tooltip = METASOUND_LOCTEXT_FORMAT(
+				"Tooltip_{0}", "Channel {0}", FText::FromString(InFriendlyName));
+			
+			return FInputDataVertex
+			{
+				InName,
+				GetMetasoundDataTypeName<FAudioBuffer>(),
+				FDataVertexMetadata
+				{
+					.Description = Tooltip,
+					.DisplayName = DisplayName,
+					.bIsAdvancedDisplay = false
+				}
+			};
 		}
 		
-		
-		FVertexInterface MakeClassInterface(const FName Format)
+		static FVertexInterface MakeClassInterface(const FName Format)
 		{
 			const Audio::FChannelTypeFamily* Found = Audio::GetChannelRegistry().FindConcreteChannel(Format);
 			if (!Found)
@@ -55,15 +57,25 @@ namespace Metasound
 			{
 				if (TOptional<Audio::FChannelTypeFamily::FChannelName> Name = Found->GetChannelName(ChannelIndex); Name.IsSet())
 				{
-					InputInterface.Add(MakeInputDataVertex(ChannelIndex, Name->Name, Name->FriendlyName));
+					InputInterface.Add(MakeInputDataVertex(Name->Name, Name->FriendlyName));
 				}
 				else
 				{
-					InputInterface.Add(MakeInputDataVertex(ChannelIndex, TEXT("Input"), FString::FromInt(ChannelIndex) ));	
+					const FName Input("In", ChannelIndex);
+					InputInterface.Add(MakeInputDataVertex(Input, FString::FromInt(ChannelIndex) ));	
 				}
 			}
+			
+			// outputs
 			FOutputVertexInterface OutputInterface;
-            OutputInterface.Add(TOutputDataVertex<FChannelAgnosticType>(METASOUND_GET_PARAM_NAME_AND_METADATA(OutputToCat)));
+			OutputInterface.Add(
+				FOutputDataVertex(
+					METASOUND_GET_PARAM_NAME(OutputToCat),
+					Format,
+					METASOUND_GET_PARAM_METADATA(OutputToCat),
+					EVertexAccessType::Reference
+				)
+			);
 			return{
 				MoveTemp(InputInterface),
 				MoveTemp(OutputInterface)
@@ -103,22 +115,26 @@ namespace Metasound
 		}
 		virtual ~FCatMakeOperator() override = default;
 
-		static const FVertexInterface& GetDefaultInterface()
+		static TDataWriteReference<FChannelAgnosticType>
+		MakePolyCat(const FName InDerivedType, const FOperatorSettings& InOpSettings)
+		{
+			using namespace Frontend;
+			const FLiteral Literal(InDerivedType.ToString());
+			TOptional<FAnyDataReference> Ref = IDataTypeRegistry::Get().CreateDataReference(InDerivedType, EDataReferenceAccessType::Write, Literal, InOpSettings);
+			return Ref->GetAs<TDataWriteReference<FChannelAgnosticType>>();
+		}
+
+		static FName GetPinName(const int32 InIndex, const Audio::FChannelTypeFamily& InFormat)
 		{
-			using namespace CatMakeNodePrivate ;
-			auto CreateDefaultInterface = []()-> FVertexInterface
+			if (auto Names = InFormat.GetChannelName(InIndex))
 			{
-				// outputs
-				FOutputVertexInterface OutputInterface;
-				OutputInterface.Add(TOutputDataVertex<FChannelAgnosticType>(METASOUND_GET_PARAM_NAME_AND_METADATA(OutputToCat)));
-				
-				return FVertexInterface({}, OutputInterface);
-			}; // end lambda: CreateDefaultInterface()
-
-			static const FVertexInterface DefaultInterface = CreateDefaultInterface();
-			return DefaultInterface;
+				return Names->Name;
+			}
+			FName Input;
+			Input.SetNumber(InIndex);
+			return Input;
 		}
-			
+		
 		static TUniquePtr<IOperator> CreateOperator(const FBuildOperatorParams& InParams, FBuildResults& OutResults)
 		{
 			using namespace CatMakeNodePrivate;
@@ -131,16 +147,20 @@ namespace Metasound
 			const FName OutputFormat = CatMakeData->GetToType();
 			const FInputVertexInterfaceData& InputData = InParams.InputData;
 			const Audio::FChannelTypeFamily* ConcreteToType = Audio::GetChannelRegistry().FindConcreteChannel(OutputFormat);
+			if (!ConcreteToType)
+			{
+				return TUniquePtr<FNoOpOperator>();
+			}
 
 			// Set the number of inputs pins to match the output format.
 			TArray<TDataReadReference<FAudioBuffer>> InputAudioVertices;
 			for (int32 i = 0; i < ConcreteToType->NumChannels(); ++i)
 			{
-				InputAudioVertices.Emplace(InputData.GetOrCreateDefaultDataReadReference<FAudioBuffer>(MakeInputVertexName(i), InParams.OperatorSettings));
+				InputAudioVertices.Emplace(InputData.GetOrCreateDefaultDataReadReference<FAudioBuffer>(GetPinName(i, *ConcreteToType), InParams.OperatorSettings));
 			}
 						
 			// Make output CAT to match our configs settings. (use concrete form).
-			FChannelAgnosticTypeWriteRef OutputPin = FChannelAgnosticTypeWriteRef::CreateNew(InParams.OperatorSettings, ConcreteToType->GetName());
+			FChannelAgnosticTypeWriteRef OutputPin = MakePolyCat(ConcreteToType->GetName(), InParams.OperatorSettings);
 			
 			return MakeUnique<FCatMakeOperator>(
 				InParams,
@@ -154,7 +174,7 @@ namespace Metasound
 			using namespace CatMakeNodePrivate;
 			for (int32 i = 0; i < InputAudioVertices.Num(); i++)
 			{
-				InOutVertexData.BindReadVertex(MakeInputVertexName(i), InputAudioVertices[i]);
+				InOutVertexData.BindReadVertex(GetPinName(i, OutputCat->GetType()), InputAudioVertices[i]);
 			}
 		}
 		virtual void BindOutputs(FOutputVertexInterfaceData& InOutVertexData) override
@@ -177,11 +197,14 @@ namespace Metasound
 				TArrayView<float> Dst = OutputCat->GetChannel(Index);
 				checkSlow(Src.Num() == Dst.Num());
 				FMemory::Memcpy(Dst.GetData(), Src.GetData(), Dst.Num() * sizeof(float));
+
+				// REWRITE ME TO BE REFERENCE FORWARD.
 			}
 		}
 		static FNodeClassMetadata GetNodeInfo()
 		{
 			using namespace CatMakeNodePrivate;
+			static FName DefaultFormat = TEXT("Cat:Mono1Dot0");
 			return FNodeClassMetadata
 			{
 				FNodeClassName{ "Experimental", "CatMakeOperator", "" },
@@ -190,8 +213,8 @@ namespace Metasound
 				METASOUND_LOCTEXT("CatMakeNodeName", "CAT Make Node"),	
 				METASOUND_LOCTEXT("CatMakeNodeNameDescription", "A Node that builds CATs"),	
 				TEXT("UE"), // Author
-				METASOUND_LOCTEXT("ExampleConfigurablePromptIfMissing", "Enable the MetaSoundExperimental Plugin"), // Prompt if missing
-				GetDefaultInterface(),
+				METASOUND_LOCTEXT("CatMakeNodePromptIfMissing", "Enable the MetaSoundExperimental Plugin"), // Prompt if missing
+				MakeClassInterface(DefaultFormat),
 				{}
 			};
 		}
diff --git a/Engine/Plugins/Experimental/MetasoundExperimental/Source/MetasoundExperimentalRuntime/Private/MetasoundCatMakeNode.h b/Engine/Plugins/Experimental/MetasoundExperimental/Source/MetasoundExperimentalRuntime/Private/MetasoundCatMakeNode.h
--- a/Engine/Plugins/Experimental/MetasoundExperimental/Source/MetasoundExperimentalRuntime/Private/MetasoundCatMakeNode.h
+++ b/Engine/Plugins/Experimental/MetasoundExperimental/Source/MetasoundExperimentalRuntime/Private/MetasoundCatMakeNode.h
@@ -14,8 +14,8 @@ struct FMetaSoundCatMakeNodeConfiguration : public FMetaSoundFrontendNodeConfigu
 
 	FMetaSoundCatMakeNodeConfiguration() = default;
 
-	UPROPERTY(EditAnywhere, Category = General, meta = (GetOptions="MetasoundCatCastingOptionsHelper.GetCastingOptions"))
-	FName Format = TEXT("Stereo");
+	UPROPERTY(EditAnywhere, Category = General, meta = (GetOptions="MetasoundCatCastingOptionsHelper.GetCastingOptions_NoAbstract"))
+	FName Format = TEXT("Cat:Stereo");
 
 	/* Get the current interface for the class based upon the node extension */
 	virtual TInstancedStruct<FMetasoundFrontendClassInterface> OverrideDefaultInterface(const FMetasoundFrontendClass& InNodeClass) const override;
diff --git a/Engine/Plugins/Experimental/MetasoundExperimental/Source/MetasoundExperimentalRuntime/Private/MetasoundCatMixerNode.cpp b/Engine/Plugins/Experimental/MetasoundExperimental/Source/MetasoundExperimentalRuntime/Private/MetasoundCatMixerNode.cpp
--- a/Engine/Plugins/Experimental/MetasoundExperimental/Source/MetasoundExperimentalRuntime/Private/MetasoundCatMixerNode.cpp
+++ b/Engine/Plugins/Experimental/MetasoundExperimental/Source/MetasoundExperimentalRuntime/Private/MetasoundCatMixerNode.cpp
@@ -2,6 +2,10 @@
 
 #include "MetasoundCatMixerNode.h"
 
+#include "Algo/MaxElement.h"
+#include "DSP/FloatArrayMath.h"
+#include "DSP/MultiMono.h"
+#include "MetasoundCatBreakNode.h"
 #include "MetasoundChannelAgnosticType.h"
 #include "MetasoundExecutableOperator.h"
 #include "MetasoundFacade.h"
@@ -9,11 +13,7 @@
 #include "MetasoundPrimitives.h"
 #include "MetasoundStandardNodesCategories.h"
 #include "MetasoundVertex.h"
-#include "Algo/MaxElement.h"
-#include "DSP/FloatArrayMath.h"
 #include "TypeFamily/ChannelTypeFamily.h"
-#include "MetasoundCatBreakNode.h"
-#include "ChannelAgnostic/ChannelAgnosticTypeUtils.h"
 
 #define LOCTEXT_NAMESPACE "MetasoundStandardNodes_CatMixerNode"
 
@@ -71,7 +71,7 @@ namespace Metasound
 			return TInputDataVertex<float>{InputName, FDataVertexMetadata{ FText::FromString(*InFriendlyName), InputDisplayName }, 1.0f};
 		}
 
-		FVertexInterface GetVertexInterface(const int32 InNumInputs)
+		static FVertexInterface GetVertexInterface(const int32 InNumInputs, const FName InOutputFormat)
 		{
 			FInputVertexInterface InputInterface;
 			for (int32 i = 0; i < InNumInputs; i++)
@@ -81,7 +81,14 @@ namespace Metasound
 			}
 
 			FOutputVertexInterface OutputInterface;
-			OutputInterface.Add(TOutputDataVertex<FChannelAgnosticType>(METASOUND_GET_PARAM_NAME_AND_METADATA(OutputCat)));
+			OutputInterface.Add(
+				FOutputDataVertex(
+					METASOUND_GET_PARAM_NAME(OutputCat),
+					InOutputFormat,
+					METASOUND_GET_PARAM_METADATA(OutputCat),
+					EVertexAccessType::Reference)
+				);
+			
 			return FVertexInterface
 				{
 					MoveTemp(InputInterface),
@@ -160,12 +167,13 @@ namespace Metasound
 
 		static const FNodeClassMetadata& GetNodeInfo()
 		{
+			static const FName DefaultOutputFormat = TEXT("Cat");
 			static const FNodeClassMetadata Metadata = CreateNodeClassMetadata
 			(
 				FName(TEXT("CAT Mixer Node")),
 				METASOUND_LOCTEXT("Metasound_CatMixerNodeDisplayName", "CAT Mixer Node"),
 				METASOUND_LOCTEXT("Metasound_CatMixerNodeDescription", "CAT Mixer Node"),
-				CatMixerPrivate::GetVertexInterface(1)
+				CatMixerPrivate::GetVertexInterface(1, DefaultOutputFormat)
 			);
 			return Metadata;
 		}
@@ -320,7 +328,10 @@ TInstancedStruct<FMetasoundFrontendClassInterface> FMetaSoundCatMixingNodeConfig
 {
 	using namespace Metasound;
 	using namespace CatMixerPrivate;
-	return TInstancedStruct<FMetasoundFrontendClassInterface>::Make(FMetasoundFrontendClassInterface::GenerateClassInterface(GetVertexInterface(NumInputs)));
+	return TInstancedStruct<FMetasoundFrontendClassInterface>::Make(
+		FMetasoundFrontendClassInterface::GenerateClassInterface(
+				GetVertexInterface(NumInputs, CustomMixFormat )
+				));
 }
 
 TSharedPtr<const Metasound::IOperatorData> FMetaSoundCatMixingNodeConfiguration::GetOperatorData() const
diff --git a/Engine/Plugins/Experimental/MetasoundExperimental/Source/MetasoundExperimentalRuntime/Private/MetasoundCatMixerNode.h b/Engine/Plugins/Experimental/MetasoundExperimental/Source/MetasoundExperimentalRuntime/Private/MetasoundCatMixerNode.h
--- a/Engine/Plugins/Experimental/MetasoundExperimental/Source/MetasoundExperimentalRuntime/Private/MetasoundCatMixerNode.h
+++ b/Engine/Plugins/Experimental/MetasoundExperimental/Source/MetasoundExperimentalRuntime/Private/MetasoundCatMixerNode.h
@@ -35,7 +35,7 @@ struct FMetaSoundCatMixingNodeConfiguration : public FMetaSoundFrontendNodeConfi
 	EMetasoundCatCastingMethod CatCastingMethod = EMetasoundCatCastingMethod::ChannelDrop;
 
 	UPROPERTY(EditAnywhere, Category = General, meta = (GetOptions="MetasoundCatCastingOptionsHelper.GetCastingOptions", EditCondition = "FormatChoosingMethod == EMetasoundMixerFormatChoosingMethod::Custom", EditConditionHides))
-	FName CustomMixFormat = TEXT("Mono");
+	FName CustomMixFormat = TEXT("Cat");
 
 	UPROPERTY(EditAnywhere, Category = General, meta = (ClampMin=1, ClampMax=100))
 	int32 NumInputs = 1;
diff --git a/Engine/Plugins/Experimental/MetasoundExperimental/Source/MetasoundExperimentalRuntime/Private/MetasoundCatPannerNode.cpp b/Engine/Plugins/Experimental/MetasoundExperimental/Source/MetasoundExperimentalRuntime/Private/MetasoundCatPannerNode.cpp
new file mode 100644
--- /dev/null
+++ b/Engine/Plugins/Experimental/MetasoundExperimental/Source/MetasoundExperimentalRuntime/Private/MetasoundCatPannerNode.cpp
@@ -0,0 +1,342 @@
+// Copyright Epic Games, Inc. All Rights Reserved.
+
+#include "MetasoundCatPannerNode.h"
+
+#include "DSP/BufferVectorOperations.h"
+#include "DSP/FloatArrayMath.h"
+#include "MetasoundChannelAgnosticType.h"
+#include "MetasoundExecutableOperator.h"
+#include "MetasoundFacade.h"
+#include "MetasoundNodeRegistrationMacro.h"
+#include "MetasoundVertex.h"
+#include "DSP/Vbap.h"
+#include "TypeFamily/ChannelTypeFamily.h"
+
+#define LOCTEXT_NAMESPACE "MetasoundStandardNodes_CatPannerNode"
+
+namespace Metasound
+{
+	class FChannelAgnosticType;
+	struct FBuildOperatorParams;
+	class FNodeFacade;
+
+	// Mapping function to convert from user facing normalized values. (0-2) 
+	// (0-1) Upper Hemisphere (Clockwise Left to Right) -90 to 90  
+	// (1-2) Lower Hemisphere (Clockwise Right to Left) 90 to -90 (sign flip through -180/180)
+	float NormalizedAzimuthToDegrees(const float InNormalizedAzimuth)
+	{
+		float WrappedAzimuth = FMath::Fmod(InNormalizedAzimuth, 2.0f);
+		if (WrappedAzimuth < 0.0f)
+		{
+			WrappedAzimuth += 2.0f;
+		}
+			
+		// Shift so 0.5 is forward (0 degrees), and scale 2.0 to (360 degrees)
+		float Degrees = (WrappedAzimuth - 0.5f) * 180.0f;
+
+		// Wrap to [-180, 180)
+		Degrees = FMath::Fmod(Degrees + 180.0f, 360.0f);
+		if (Degrees < 0.0f)
+		{
+			Degrees += 360.0f;
+		}
+		Degrees -= 180.0f;
+			
+		return Degrees;
+	}
+
+	namespace CatPannerPrivate
+	{
+		METASOUND_PARAM(InputAudio,			"Input",	"Input audio to pan.");
+		METASOUND_PARAM(InputAzimuth,		"Azimuth",	"The normalized azimuthal angle to pan the sound in range [0.0, 2.0). (0.0 is left, 0.5 is center, 1.0 is right, 1.5 is behind.");
+		METASOUND_PARAM(OutputToCat,		"Output",	"Panned audio.");
+		
+		class FCatPannerOperatorData final : public TOperatorData<FCatPannerOperatorData>
+		{
+		public:
+			static const FLazyName OperatorDataTypeName;
+
+			explicit FCatPannerOperatorData(const FName& InPanToTypeName, ECatPannerMethod InPannerMethod)
+				: PanToTypeName(InPanToTypeName)
+				, PannerMethod(InPannerMethod)
+			{}
+			const FName& GetPanToType() const { return PanToTypeName; }
+			ECatPannerMethod GetPannerMethod() const { return PannerMethod; }
+		private:
+			FName PanToTypeName;
+			ECatPannerMethod PannerMethod;
+		};
+
+		// Linkage.
+		const FLazyName FCatPannerOperatorData::OperatorDataTypeName = TEXT("FCatPannerOperatorData");
+
+	}
+	
+	class FCatPannerOperator final : public TExecutableOperator<FCatPannerOperator>
+	{
+	public:
+		using FDiscreteChannelTypeFamily = Audio::FDiscreteChannelTypeFamily;
+		using FCatPannerOperatorData = CatPannerPrivate::FCatPannerOperatorData;
+
+		FCatPannerOperator(
+				const FBuildOperatorParams& InParams, 
+				FAudioBufferReadRef&& InAudio, 
+				FFloatReadRef&& InAzimuth,
+				FDiscreteChannelAgnosticTypeWriteRef&& InCatOutput,
+				const Audio::IDiscretePanner& InPanner,
+				const FCatPannerOperatorData* InConfigData)
+			: InputAudio(MoveTemp(InAudio))
+			, InputAzimuth(MoveTemp(InAzimuth))
+			, PannerMethod(InConfigData->GetPannerMethod())
+			, OutputPannedAudio(MoveTemp(InCatOutput))
+			, Settings(InParams.OperatorSettings)
+			, Panner(InPanner)
+		{
+			LastUpdateGainsPerChannel.SetNum(InCatOutput->NumChannels());
+			for (int32 i = 0; i < InCatOutput->NumChannels(); ++i)
+			{
+				LastUpdateGainsPerChannel[i] = -1.f;
+			}
+		}
+		
+		virtual ~FCatPannerOperator() override = default;
+
+		static FVertexInterface GetInterface(const FName OutputFormat)
+		{
+			using namespace CatPannerPrivate;
+			// inputs
+			FInputVertexInterface InputInterface;
+			InputInterface.Add(TInputDataVertex<FAudioBuffer>(METASOUND_GET_PARAM_NAME_AND_METADATA(InputAudio)));
+			InputInterface.Add(TInputDataVertex<float>(METASOUND_GET_PARAM_NAME_AND_METADATA(InputAzimuth), 0.0f));
+
+			// outputs
+			FOutputVertexInterface OutputInterface;
+			OutputInterface.Add(
+				FOutputDataVertex(
+					METASOUND_GET_PARAM_NAME(OutputToCat),
+					OutputFormat,
+					METASOUND_GET_PARAM_METADATA(OutputToCat),
+					EVertexAccessType::Reference)
+				);
+
+			return FVertexInterface(InputInterface, OutputInterface);
+		}
+			
+		static TUniquePtr<IOperator> CreateOperator(const FBuildOperatorParams& InParams, FBuildResults& OutResults)
+		{
+			using namespace CatPannerPrivate;
+
+			const FCatPannerOperatorData* ConfigData = CastOperatorData<const FCatPannerOperatorData>(InParams.Node.GetOperatorData().Get());
+			if (!ConfigData)
+			{
+				return MakeUnique<FNoOpOperator>();
+			}
+			const FName PanToName = ConfigData->GetPanToType();
+			const Audio::FChannelTypeFamily* ConcretePanToType = Audio::GetChannelRegistry().FindConcreteChannel(PanToName);
+			if (!ConcretePanToType)
+			{
+				return MakeUnique<FNoOpOperator>();
+			}
+
+			TDataReadReference<FAudioBuffer> InputAudio = 
+				InParams.InputData.GetOrCreateDefaultDataReadReference<FAudioBuffer>(METASOUND_GET_PARAM_NAME(InputAudio), InParams.OperatorSettings);
+			
+			TDataReadReference<float> InputAzimuth =
+				InParams.InputData.GetOrCreateDefaultDataReadReference<float>(METASOUND_GET_PARAM_NAME(InputAzimuth), InParams.OperatorSettings);
+
+			// Make Discrete CAT output.
+			// These are defined in data so we must use "DerivedAs" 
+			TDataWriteReference<FDiscreteChannelAgnosticType> CatPin =
+				FDiscreteChannelAgnosticTypeWriteRef::CreateNewDerivedAs<FDiscreteChannelAgnosticType>(
+					ConcretePanToType->GetName(),
+					Metasound::GetMetasoundDataTypeId<FDiscreteChannelAgnosticType>(),
+					InParams.OperatorSettings,
+					ConcretePanToType->GetName()
+				);
+
+			// Ask target format for a panner. 
+			const Audio::FChannelTypeFamily::FGetPannerParams Params;
+			const Audio::IDiscretePanner* InPanner = CatPin->GetType().GetPanner(Params);
+			if (!InPanner)
+			{
+				return MakeUnique<FNoOpOperator>();
+			}
+			
+			return MakeUnique<FCatPannerOperator>(
+				InParams,
+				MoveTemp(InputAudio),
+				MoveTemp(InputAzimuth),
+				MoveTemp(CatPin),
+				*InPanner,
+				ConfigData);
+		}
+
+		virtual void BindInputs(FInputVertexInterfaceData& InOutVertexData) override
+		{
+			using namespace CatPannerPrivate;
+			InOutVertexData.BindReadVertex(METASOUND_GET_PARAM_NAME(InputAudio) , InputAudio);
+			InOutVertexData.BindReadVertex(METASOUND_GET_PARAM_NAME(InputAzimuth), InputAzimuth);
+		}
+		
+		virtual void BindOutputs(FOutputVertexInterfaceData& InOutVertexData) override
+		{
+			using namespace CatPannerPrivate;
+			InOutVertexData.BindWriteVertex(METASOUND_GET_PARAM_NAME(OutputToCat), OutputPannedAudio);
+		}
+
+		void Reset(const FResetParams& InParams)
+		{
+			Execute();
+		}
+
+		void ApplyGains(const Audio::IDiscretePanner::FOutputParams& InOutput)
+		{
+			const FDiscreteChannelTypeFamily& Discrete = *OutputPannedAudio->GetType().Cast<FDiscreteChannelTypeFamily>();
+
+			const int32 InputSampleCount = InputAudio->Num();
+
+			auto ApplyGain = [](float& Prev, const float Current, TArrayView<const float> Src, TArrayView<float> Dst)
+				{
+					if (FMath::IsNearlyZero(Current))
+					{
+						// Nothing so far and zero gain. No nothing.
+						if (Prev < 0.f) 
+						{
+							return;
+						}
+
+						// Zero start+end. (we can zero).
+						if ( FMath::IsNearlyZero(Prev)) 
+						{
+							FMemory::Memzero(Dst.GetData(), Dst.NumBytes());
+							return;
+						}
+					}
+
+					// Special case if prev is 100% and current is 100%, just copy source out.
+					if (FMath::IsNearlyEqual(Prev,1.0f) && FMath::IsNearlyEqual(Current,1.0f))
+					{
+						FMemory::Memcpy(Dst.GetData(), Src.GetData(), Src.NumBytes());
+					}
+					else
+					{
+						// Clamp previous >= 0
+						const float PrevClamped = FMath::Max(0,Prev);
+						
+						// Requires some kind of fade...
+						Audio::ArrayFade(Src, PrevClamped, Current, Dst);
+					}
+
+					// Remember the gain as previous for next round.
+					Prev = Current;
+				};
+
+			
+			const int32 NumChannels = OutputPannedAudio->NumChannels();
+			for (int32 i = 0; i < NumChannels; ++i)
+			{
+				// If this channel is active in the pan results.
+				const Audio::IDiscretePanner::FPanResult* ActiveSpeaker = InOutput.Results.FindByPredicate(
+					[i](const Audio::IDiscretePanner::FPanResult& PanResult) -> bool
+						{
+							check(PanResult.ChannelIndex != INDEX_NONE);
+							return PanResult.ChannelIndex == i;
+						});
+
+				const float TargetGain = ActiveSpeaker ? ActiveSpeaker->Gain : 0.f;
+				ApplyGain(LastUpdateGainsPerChannel[i], TargetGain, *InputAudio, OutputPannedAudio->GetChannel(i));
+			}
+		}
+	
+		
+		void Execute()
+		{
+			// 0-1 (upper hemisphere clockwise  Left+Right) -90 to 90 (9pm - 3pm)
+			// 1-2 (lower hemisphere clockwise Right-Left) 90 to -90  (3pm - 9pm)
+
+			const float AzimuthDegrees = NormalizedAzimuthToDegrees(*InputAzimuth);
+
+			// Zero output.
+			// OutputPannedAudio->Zero();
+			
+			Audio::IDiscretePanner::FOutputParams OutputParam;
+			Panner.ComputeGains(
+				{
+					.AzimuthDegrees = AzimuthDegrees,
+					.ElevationDegrees = 0.f,
+					.bAllowAzimuthMirroring = true
+				},
+				OutputParam
+			);
+
+			// Apply gains.
+			ApplyGains(OutputParam);
+		}
+
+		static FNodeClassMetadata GetNodeInfo()
+		{
+			using namespace CatPannerPrivate;
+			static const FName DefaultFormat(TEXT("Cat:Stereo2Dot0"));
+			static const FText MissingText = LOCTEXT("ExampleCatPannerPromptIfMissing", "Enable the MetaSoundExperimental Plugin");
+			static const FText NodeNameText = LOCTEXT("CatPannerNodeName", "Azimuth Panner");
+			static const FText NodeDescText = LOCTEXT("CatPannerNodeNameDescription", "A Node that does an azimuthal pan to an output CAT type");
+			
+		return FNodeClassMetadata
+			{
+				FNodeClassName{ "Experimental", "CatPannerOperator", "" },
+				1, // Major version
+				0, // Minor version
+				NodeNameText,	
+				NodeDescText,	
+				TEXT("UE"), // Author
+				MissingText, // Prompt if missing
+				GetInterface(DefaultFormat),
+				{}
+			};
+		}
+	private:
+		FAudioBufferReadRef InputAudio;
+		FFloatReadRef InputAzimuth;
+		ECatPannerMethod PannerMethod;
+		FDiscreteChannelAgnosticTypeWriteRef OutputPannedAudio;
+		FOperatorSettings Settings;
+		const Audio::IDiscretePanner& Panner;
+		TArray<float> LastUpdateGainsPerChannel;
+		TMap<ESpeakerShortNames, int32> SpeakerMapping;
+	}; 
+		
+	using FCatPannerNode = TNodeFacade<FCatPannerOperator>;
+
+	// register node config.	
+	METASOUND_REGISTER_NODE_AND_CONFIGURATION(FCatPannerNode, FMetaSoundCatPannerNodeConfiguration);
+} // namespace Metasound
+
+TArray<FPropertyTextFName> UMetasoundCatAzimuthalChannelOptionsHelper::GetAzimuthalChannelOptions()
+{
+	using namespace Audio;
+	TArray<TSharedRef<const FChannelTypeFamily>> AllFormats = Audio::GetChannelRegistry().GetAllChannelFormats();
+	TArray<FPropertyTextFName> FormatsOptions;
+	for (TSharedRef Format : AllFormats)
+	{
+		if (const FDiscreteChannelTypeFamily* Discrete = Format->Cast<FDiscreteChannelTypeFamily>())
+		{
+			if (Discrete->GetPanner() != nullptr) // If there's a panner defined, we can use it.
+			{
+				FormatsOptions.Emplace(Format->GetName(), FText::FromString(Format->GetFriendlyName()));
+			}
+		}
+	}
+	return FormatsOptions;
+}
+
+TInstancedStruct<FMetasoundFrontendClassInterface> FMetaSoundCatPannerNodeConfiguration::OverrideDefaultInterface(const FMetasoundFrontendClass& InNodeClass) const
+{
+	return TInstancedStruct<FMetasoundFrontendClassInterface>::Make(FMetasoundFrontendClassInterface::GenerateClassInterface(Metasound::FCatPannerOperator::GetInterface(this->PanToType)));		// Assume this is concrete type.
+}
+
+TSharedPtr<const Metasound::IOperatorData> FMetaSoundCatPannerNodeConfiguration::GetOperatorData() const
+{
+	return MakeShared<Metasound::CatPannerPrivate::FCatPannerOperatorData>(PanToType, PanningMethod);
+}
+
+#undef LOCTEXT_NAMESPACE // 
diff --git a/Engine/Plugins/Experimental/MetasoundExperimental/Source/MetasoundExperimentalRuntime/Private/MetasoundChannelAgnosticType.cpp b/Engine/Plugins/Experimental/MetasoundExperimental/Source/MetasoundExperimentalRuntime/Private/MetasoundChannelAgnosticType.cpp
--- a/Engine/Plugins/Experimental/MetasoundExperimental/Source/MetasoundExperimentalRuntime/Private/MetasoundChannelAgnosticType.cpp
+++ b/Engine/Plugins/Experimental/MetasoundExperimental/Source/MetasoundExperimentalRuntime/Private/MetasoundChannelAgnosticType.cpp
@@ -3,6 +3,7 @@
 #include "MetasoundChannelAgnosticType.h"
 #include "TypeFamily/TypeFamily.h"
 #include "TypeFamily/ChannelTypeFamily.h"
+#include "MetasoundPolymorphic.h"
 
 namespace Metasound
 {
@@ -19,7 +20,7 @@ namespace Metasound
 	 */
 	FName FChannelAgnosticType::GetDefaultCatFormat()
 	{
-		static const FName DefaultFormat(TEXT("Mono")); // 1 channel, Mono buffer.
+		static const FName DefaultFormat(TEXT("Cat:Mono1Dot0"));
 		check(Audio::GetChannelRegistry().FindConcreteChannel(DefaultFormat) != nullptr);
 		return DefaultFormat;
 	}
@@ -29,9 +30,25 @@ namespace Metasound
 	{}
 	
 	FChannelAgnosticType::FChannelAgnosticType(const int32 InNumFramesPerBlock, const FName& InChannelTypeName)
-		: Super(*Audio::GetChannelRegistry().FindConcreteChannel(InChannelTypeName), InNumFramesPerBlock)
+		: Super(Audio::GetChannelRegistry().FindConcreteChannelChecked(InChannelTypeName), InNumFramesPerBlock)
 	{}
+	
+	// Register.
+	REGISTER_METASOUND_DATATYPE(FChannelAgnosticType, "Cat");
+	REGISTER_METASOUND_POLY_TYPE(FChannelAgnosticType);
+
+	REGISTER_METASOUND_DATATYPE(FCompositeChannelAgnosticType, "Cat:Composite");
+	REGISTER_METASOUND_POLY_TYPE(FCompositeChannelAgnosticType);
+
+	REGISTER_METASOUND_DATATYPE(FDiscreteChannelAgnosticType, "Cat:Discrete");
+	REGISTER_METASOUND_POLY_TYPE(FDiscreteChannelAgnosticType);
+
+	REGISTER_METASOUND_DATATYPE(FSoundfieldChannelAgnosticType, "Cat:Soundfield");
+	REGISTER_METASOUND_POLY_TYPE(FSoundfieldChannelAgnosticType);
 }
 
-// Register.
-REGISTER_METASOUND_DATATYPE(Metasound::FChannelAgnosticType, "Cat");
\ No newline at end of file
+
+
+
+
+
diff --git a/Engine/Plugins/Experimental/MetasoundExperimental/Source/MetasoundExperimentalRuntime/Private/MetasoundExperimentalRuntimeModule.cpp b/Engine/Plugins/Experimental/MetasoundExperimental/Source/MetasoundExperimentalRuntime/Private/MetasoundExperimentalRuntimeModule.cpp
--- a/Engine/Plugins/Experimental/MetasoundExperimental/Source/MetasoundExperimentalRuntime/Private/MetasoundExperimentalRuntimeModule.cpp
+++ b/Engine/Plugins/Experimental/MetasoundExperimental/Source/MetasoundExperimentalRuntime/Private/MetasoundExperimentalRuntimeModule.cpp
@@ -2,12 +2,17 @@
 
 #include "MetasoundExperimentalRuntimeModule.h"
 
+#include "MetasoundCatConversionNode.h"
+#include "MetasoundCatDataSchemas.h"
 #include "MetasoundFrontendModuleRegistrationMacros.h"
 #include "Modules/ModuleManager.h"
 
 void FMetasoundExperimentalRuntimeModule::StartupModule()
 {
 	FModuleManager::Get().LoadModule("AudioExperimentalRuntime");
+	Metasound::LoadPolymorphicDataSchemas();
+	Metasound::RegisterCatConvertNode();
+	
 	METASOUND_REGISTER_ITEMS_IN_MODULE
 }
 
diff --git a/Engine/Plugins/Experimental/MetasoundExperimental/Source/MetasoundExperimentalRuntime/Public/MetasoundCatPannerNode.h b/Engine/Plugins/Experimental/MetasoundExperimental/Source/MetasoundExperimentalRuntime/Public/MetasoundCatPannerNode.h
new file mode 100644
--- /dev/null
+++ b/Engine/Plugins/Experimental/MetasoundExperimental/Source/MetasoundExperimentalRuntime/Public/MetasoundCatPannerNode.h
@@ -0,0 +1,58 @@
+// Copyright Epic Games, Inc. All Rights Reserved.
+
+#pragma once
+
+#include "MetasoundFrontendDocument.h"
+#include "Templates/SharedPointer.h"
+#include "UObject/PropertyText.h"
+
+#include "MetasoundCatPannerNode.generated.h"
+
+UCLASS()
+class UMetasoundCatAzimuthalChannelOptionsHelper : public UObject
+{
+	GENERATED_BODY()
+
+	UFUNCTION()
+	static TArray<FPropertyTextFName> GetAzimuthalChannelOptions();
+};
+
+UENUM()
+enum class ECatPannerMethod : uint8
+{
+	// Panning is done using equal-power law
+	EqualPower,
+
+	// Panning is done linearly
+	Linear
+};
+
+
+USTRUCT()
+struct FMetaSoundCatPannerNodeConfiguration : public FMetaSoundFrontendNodeConfiguration
+{
+	GENERATED_BODY()
+
+	FMetaSoundCatPannerNodeConfiguration() =  default;
+
+	// What the CAT output type is. Only azimuthal discrete channel types supported.
+	UPROPERTY(EditAnywhere, Category = General, meta = (GetOptions="MetasoundCatAzimuthalChannelOptionsHelper.GetAzimuthalChannelOptions"))
+	FName PanToType = TEXT("Cat:Stereo2Dot0");
+
+	// What panning method to use. Panning method defines the mapping function of gains between pair-wise channels.
+	UPROPERTY(EditAnywhere, Category = General)
+	ECatPannerMethod PanningMethod = ECatPannerMethod::EqualPower;
+	
+	/* Get the current interface for the class based upon the node extension */
+	virtual TInstancedStruct<FMetasoundFrontendClassInterface> OverrideDefaultInterface(const FMetasoundFrontendClass& InNodeClass) const override;
+
+	/* Pass data down to the operator. */
+	virtual TSharedPtr<const Metasound::IOperatorData> GetOperatorData() const override;
+};
+
+
+namespace Metasound
+{
+	// Normalization utility.	
+	METASOUNDEXPERIMENTALRUNTIME_API float NormalizedAzimuthToDegrees(const float InNormalizedAzimuth);
+}
diff --git a/Engine/Plugins/Experimental/MetasoundExperimental/Source/MetasoundExperimentalRuntime/Public/MetasoundChannelAgnosticType.h b/Engine/Plugins/Experimental/MetasoundExperimental/Source/MetasoundExperimentalRuntime/Public/MetasoundChannelAgnosticType.h
--- a/Engine/Plugins/Experimental/MetasoundExperimental/Source/MetasoundExperimentalRuntime/Public/MetasoundChannelAgnosticType.h
+++ b/Engine/Plugins/Experimental/MetasoundExperimental/Source/MetasoundExperimentalRuntime/Public/MetasoundChannelAgnosticType.h
@@ -2,9 +2,11 @@
 
 #pragma once
 
+#include "MetasoundPolymorphic.h"
 #include "ChannelAgnostic/ChannelAgnosticType.h"
 #include "MetasoundDataTypeRegistrationMacro.h"
-#include "DSP/MultiMono.h"
+
+#include "MetasoundChannelAgnosticType.generated.h"
 
 namespace Metasound
 {
@@ -20,22 +22,119 @@ namespace Metasound
 		 */
 		static FName GetDefaultCatFormat();
 
-		explicit FChannelAgnosticType(const FOperatorSettings& InSettings, const FName& InChannelTypeName = GetDefaultCatFormat());
+		METASOUNDEXPERIMENTALRUNTIME_API explicit FChannelAgnosticType(const FOperatorSettings& InSettings, const FName& InChannelTypeName = GetDefaultCatFormat());
 		
 		explicit FChannelAgnosticType(const int32 InNumFramesPerBlock, const FName& InChannelTypeName = GetDefaultCatFormat());
-				
+
+		static FName LiteralToFName(const FString& In)
+		{
+			if (In.IsEmpty())
+			{
+				static const FName Default = TEXT("Cat:Mono1Dot0");
+				return Default;
+			}
+			return FName(In);
+		}
 		/* FLiteral compatible Constructor
 		 * TODO: Literals only support FString and not FNames, so there is a conversion here.
 		 */
-		explicit FChannelAgnosticType(const FOperatorSettings& InSettings, const FString& InChannelTypeName)
-			: FChannelAgnosticType(InSettings, FName(*InChannelTypeName))
+		METASOUNDEXPERIMENTALRUNTIME_API explicit FChannelAgnosticType(const FOperatorSettings& InSettings, const FString& InChannelTypeName)
+			: FChannelAgnosticType(InSettings, LiteralToFName((InChannelTypeName)))
 		{}
 	};
-	
 
-	
+	class FDiscreteChannelAgnosticType : public FChannelAgnosticType
+	{
+		using Super = FChannelAgnosticType;
+		using Super::Super;
+	};
+	class FSoundfieldChannelAgnosticType : public FChannelAgnosticType
+	{
+		using Super = FChannelAgnosticType;
+		using Super::Super;
+	};
+	class FCompositeChannelAgnosticType : public FChannelAgnosticType
+	{
+		using Super = FChannelAgnosticType;
+		using Super::Super;
+	};
 }
 
-// Declare it.
-DECLARE_METASOUND_DATA_REFERENCE_TYPES(Metasound::FChannelAgnosticType, METASOUNDEXPERIMENTALRUNTIME_API,
-	FChannelAgnosticTypeTypeInfo, FChannelAgnosticTypeReadRef, FChannelAgnosticTypeWriteRef);
+USTRUCT()
+struct FCatDataTypeSchema
+{
+	GENERATED_BODY()
+
+	UPROPERTY()
+	FName Name;
+
+	UPROPERTY()
+	FName BaseType;
+
+	UPROPERTY()
+	FName FamilyType;
+
+	UPROPERTY()
+	FString FriendlyName;
+
+	UPROPERTY()
+	bool bIsAbstract = false;
+
+	UPROPERTY()
+	bool bIsParentsDefault = false;
+};
+
+USTRUCT()
+struct FCatSoundfieldDataSchema : public FCatDataTypeSchema
+{
+	GENERATED_BODY()
+
+	UPROPERTY()
+	int32 NumOrders = 1;
+};
+
+USTRUCT()
+struct FCatDiscreteChannel
+{
+	GENERATED_BODY()
+
+	UPROPERTY()
+	float ElevationDegrees = 0;			// In Degrees. +180 to -180
+
+	UPROPERTY()
+	float AzimuthDegrees = 0;			// -180 to 180 azimuth.
+
+	UPROPERTY()
+	FName ID;							// This should really be the speaker enum.
+
+	UPROPERTY()
+	bool bIsSpatialized = true;			// If this is off, azimuth+degrees do nothing. (for excluding from panner). 
+};
+
+USTRUCT()
+struct FCatDiscreteDataSchema : public FCatDataTypeSchema
+{
+	GENERATED_BODY()
+
+	UPROPERTY()
+	TArray<FCatDiscreteChannel> Order;
+};
+
+USTRUCT()
+struct FCatCompositeDataSchema : public FCatDataTypeSchema
+{
+	GENERATED_BODY()
+};
+
+namespace  Metasound
+{
+	DECLARE_METASOUND_DATA_REFERENCE_TYPES(Metasound::FChannelAgnosticType,				METASOUNDEXPERIMENTALRUNTIME_API, FChannelAgnosticTypeTypeInfo,				FChannelAgnosticTypeReadRef,			FChannelAgnosticTypeWriteRef);
+	DECLARE_METASOUND_DATA_REFERENCE_TYPES(Metasound::FCompositeChannelAgnosticType,	METASOUNDEXPERIMENTALRUNTIME_API, FCompositeChannelAgnosticTypeTypeInfo,	FCompositeChannelAgnosticTypeReadRef,	FCompositeChannelAgnosticTypeWriteRef);
+	DECLARE_METASOUND_DATA_REFERENCE_TYPES(Metasound::FDiscreteChannelAgnosticType,		METASOUNDEXPERIMENTALRUNTIME_API, FDiscreteChannelAgnosticTypeTypeInfo,		FDiscreteChannelAgnosticTypeReadRef,	FDiscreteChannelAgnosticTypeWriteRef);
+	DECLARE_METASOUND_DATA_REFERENCE_TYPES(Metasound::FSoundfieldChannelAgnosticType,	METASOUNDEXPERIMENTALRUNTIME_API, FSoundfieldChannelAgnosticTypeTypeInfo,	FSoundfieldChannelAgnosticTypeReadRef,	FSoundChannelAgnosticTypeWriteRef);
+
+	DECLARE_METASOUND_POLY_TYPE(FChannelAgnosticType,			void, true);
+	DECLARE_METASOUND_POLY_TYPE(FCompositeChannelAgnosticType,	Metasound::FChannelAgnosticType, true);
+	DECLARE_METASOUND_POLY_TYPE(FDiscreteChannelAgnosticType,	Metasound::FChannelAgnosticType, true);
+	DECLARE_METASOUND_POLY_TYPE(FSoundfieldChannelAgnosticType,	Metasound::FChannelAgnosticType, true);
+}
\ No newline at end of file
diff --git a/Engine/Plugins/Runtime/Metasound/Source/MetasoundFrontend/Private/MetasoundFrontendDataTypeRegistry.cpp b/Engine/Plugins/Runtime/Metasound/Source/MetasoundFrontend/Private/MetasoundFrontendDataTypeRegistry.cpp
--- a/Engine/Plugins/Runtime/Metasound/Source/MetasoundFrontend/Private/MetasoundFrontendDataTypeRegistry.cpp
+++ b/Engine/Plugins/Runtime/Metasound/Source/MetasoundFrontend/Private/MetasoundFrontendDataTypeRegistry.cpp
@@ -719,15 +719,19 @@ PRAGMA_ENABLE_DEPRECATION_WARNINGS
 					{
 						if (Entry->GetDataTypeInfo().bIsParsable)
 						{
-							NodeRegistry->RegisterNode(MakeUnique<FInputNodeRegistryEntry>(Entry));
+							const bool bIsAbstract = Entry->GetDataTypeInfo().bIsAbstract;;
 							
+							NodeRegistry->RegisterNode(MakeUnique<FInputNodeRegistryEntry>(Entry));							
 							NodeRegistry->RegisterNode(MakeUnique<FOutputNodeRegistryEntry>(Entry));
-							
-							NodeRegistry->RegisterNode(MakeUnique<FLiteralNodeRegistryEntry>(Entry));
-							NodeRegistry->RegisterNode(MakeUnique<FVariableNodeRegistryEntry>(Entry));
-							NodeRegistry->RegisterNode(MakeUnique<FVariableMutatorNodeRegistryEntry>(Entry));
-							NodeRegistry->RegisterNode(MakeUnique<FVariableAccessorNodeRegistryEntry>(Entry));
-							NodeRegistry->RegisterNode(MakeUnique<FVariableDeferredAccessorNodeRegistryEntry>(Entry));
+
+							if (!bIsAbstract)
+							{
+								NodeRegistry->RegisterNode(MakeUnique<FLiteralNodeRegistryEntry>(Entry));
+								NodeRegistry->RegisterNode(MakeUnique<FVariableNodeRegistryEntry>(Entry));
+								NodeRegistry->RegisterNode(MakeUnique<FVariableMutatorNodeRegistryEntry>(Entry));
+								NodeRegistry->RegisterNode(MakeUnique<FVariableAccessorNodeRegistryEntry>(Entry));
+								NodeRegistry->RegisterNode(MakeUnique<FVariableDeferredAccessorNodeRegistryEntry>(Entry));
+							}
 
 							if (Entry->GetDataTypeInfo().bIsConstructorType)
 							{
diff --git a/Engine/Plugins/Runtime/Metasound/Source/MetasoundFrontend/Private/MetasoundOutputNode.cpp b/Engine/Plugins/Runtime/Metasound/Source/MetasoundFrontend/Private/MetasoundOutputNode.cpp
--- a/Engine/Plugins/Runtime/Metasound/Source/MetasoundFrontend/Private/MetasoundOutputNode.cpp
+++ b/Engine/Plugins/Runtime/Metasound/Source/MetasoundFrontend/Private/MetasoundOutputNode.cpp
@@ -67,7 +67,7 @@ namespace Metasound
 			// Use data reference if it is passed in. 
 			if (const FAnyDataReference* Ref = InParams.InputData.FindDataReference(VertexName))
 			{
-				checkf(Ref->GetDataTypeName() == Vertex.DataTypeName, TEXT("Mismatched data type names on output node (%s). Expected (%s), received (%s)."), *VertexName.ToString(), *Vertex.DataTypeName.ToString(), *(Ref->GetDataTypeName().ToString()));
+				checkf(IsCastable(Ref->GetDataTypeName(), Vertex.DataTypeName), TEXT("Mismatched data type names on output node (%s). Expected (%s), received (%s)."), *VertexName.ToString(), *Vertex.DataTypeName.ToString(), *(Ref->GetDataTypeName().ToString()));
 				return MakeUnique<FOutputOperator>(VertexName, *Ref);
 			}
 			else 
diff --git a/Engine/Plugins/Runtime/Metasound/Source/MetasoundFrontend/Public/MetasoundNodeConstructorParams.h b/Engine/Plugins/Runtime/Metasound/Source/MetasoundFrontend/Public/MetasoundNodeConstructorParams.h
--- a/Engine/Plugins/Runtime/Metasound/Source/MetasoundFrontend/Public/MetasoundNodeConstructorParams.h
+++ b/Engine/Plugins/Runtime/Metasound/Source/MetasoundFrontend/Public/MetasoundNodeConstructorParams.h
@@ -40,6 +40,9 @@ namespace Metasound
 		FVertexName VertexName;
 		FLiteral InitParam = FLiteral::CreateInvalid();
 
+		// If defined this is the polymorphic type name to use.
+		FName PolymorphicTypeName;
+
 		UE_DEPRECATED(5.3, "Transmission is no longer enabled on input nodes. In order set values on a MetaSound input, access them through the bound FInputVertexInterfaceData.")
 		bool bEnableTransmission = false;
 	};
diff --git a/Engine/Plugins/Runtime/Metasound/Source/MetasoundGraphCore/Private/MetasoundOperatorBuilder.cpp b/Engine/Plugins/Runtime/Metasound/Source/MetasoundGraphCore/Private/MetasoundOperatorBuilder.cpp
--- a/Engine/Plugins/Runtime/Metasound/Source/MetasoundGraphCore/Private/MetasoundOperatorBuilder.cpp
+++ b/Engine/Plugins/Runtime/Metasound/Source/MetasoundGraphCore/Private/MetasoundOperatorBuilder.cpp
@@ -680,7 +680,7 @@ PRAGMA_ENABLE_DEPRECATION_WARNINGS
 
 				if (const FAnyDataReference* DataReference = NodeInputData.FindDataReference(InputDestination.Vertex.VertexName))
 				{
-					if (DataReference->GetDataTypeName() == InputDestination.Vertex.DataTypeName)
+					if (IsCastable(DataReference->GetDataTypeName(), InputDestination.Vertex.DataTypeName))
 					{
 						bFoundDataReference = true;
 						OutVertexData.GetInputs().SetVertex(InputDestination.Vertex.VertexName, *DataReference);
@@ -708,7 +708,7 @@ PRAGMA_ENABLE_DEPRECATION_WARNINGS
 				const FOutputVertexInterfaceData& NodeOutputData = OperatorInfo->VertexData.GetOutputs();
 				if (const FAnyDataReference* DataReference = NodeOutputData.FindDataReference(OutputSource.Vertex.VertexName))
 				{
-					if (DataReference->GetDataTypeName() == OutputSource.Vertex.DataTypeName)
+					if (IsCastable(DataReference->GetDataTypeName() , OutputSource.Vertex.DataTypeName))
 					{
 						bFoundDataReference = true;
 						OutVertexData.GetOutputs().SetVertex(OutputSource.Vertex.VertexName, FAnyDataReference{*DataReference});
diff --git a/Engine/Plugins/Runtime/Metasound/Source/MetasoundGraphCore/Private/MetasoundVertexData.cpp b/Engine/Plugins/Runtime/Metasound/Source/MetasoundGraphCore/Private/MetasoundVertexData.cpp
--- a/Engine/Plugins/Runtime/Metasound/Source/MetasoundGraphCore/Private/MetasoundVertexData.cpp
+++ b/Engine/Plugins/Runtime/Metasound/Source/MetasoundGraphCore/Private/MetasoundVertexData.cpp
@@ -81,7 +81,7 @@ namespace Metasound
 
 		void FInputBinding::Set(FAnyDataReference&& InAnyDataReference)
 		{
-			check(IsCastable(Vertex.DataTypeName,InAnyDataReference.GetDataTypeName()));
+			check(IsCastable(InAnyDataReference.GetDataTypeName(),Vertex.DataTypeName));
 			Data.Emplace(MoveTemp(InAnyDataReference));
 			CheckAccessTypeCompatibility(Vertex, *Data);
 		}
@@ -128,7 +128,7 @@ namespace Metasound
 
 		void FInputBinding::Bind(FInputBinding& InBinding)
 		{
-			check(IsCastable(Vertex.DataTypeName, InBinding.GetVertex().DataTypeName));
+			check(IsCastable(InBinding.GetVertex().DataTypeName,Vertex.DataTypeName));
 
 			if (Data.IsSet()) 
 			{ 
@@ -166,7 +166,7 @@ namespace Metasound
 
 		void FOutputBinding::Set(FAnyDataReference&& InAnyDataReference)
 		{
-			check(IsCastable(Vertex.DataTypeName, InAnyDataReference.GetDataTypeName()));
+			check(IsCastable(InAnyDataReference.GetDataTypeName(),Vertex.DataTypeName));
 			Data.Emplace(MoveTemp(InAnyDataReference));
 			CheckAccessTypeCompatibility(Vertex, *Data);
 		}
@@ -251,7 +251,10 @@ namespace Metasound
 		{
 			if (BindingType* Binding = Find(InBindings, InVertexName))
 			{
-				if (IsCastable(Binding->GetVertex().DataTypeName,InDataReference.GetDataTypeName()))
+				if (IsCastable(
+					InDataReference.GetDataTypeName(),						// Incoming ref. (dynamic reference).
+					Binding->GetVertex().DataTypeName						// Type of this pin. (must be derived from this or equal).
+					))
 				{
 					Binding->Set(MoveTemp(InDataReference));
 				}
diff --git a/Engine/Plugins/Runtime/Metasound/Source/MetasoundGraphCore/Public/MetasoundDataReference.h b/Engine/Plugins/Runtime/Metasound/Source/MetasoundGraphCore/Public/MetasoundDataReference.h
--- a/Engine/Plugins/Runtime/Metasound/Source/MetasoundGraphCore/Public/MetasoundDataReference.h
+++ b/Engine/Plugins/Runtime/Metasound/Source/MetasoundGraphCore/Public/MetasoundDataReference.h
@@ -245,6 +245,30 @@ namespace Metasound
 			);
 		}
 
+		template <typename Derived, typename... ArgTypes>
+		UE_EXPERIMENTAL(5.8, "Polymorphic API is under development")
+		static TDataReference CreateNewDerivedAs(const FName& InDerivedTypeName,
+												 FMetasoundDataTypeId InDerivedTypeId,
+												 ArgTypes&&... Args)
+		{
+			static_assert(std::is_base_of_v<DataType, Derived>, "Derived must inherit from DataType");
+			static_assert(std::is_constructible_v<Derived, ArgTypes...>, "Args don't match a Derived constructor");
+
+		#if DO_CHECK
+			PRAGMA_DISABLE_EXPERIMENTAL_WARNINGS
+			check(Metasound::IsChildOfByName(InDerivedTypeName, GetMetasoundDataTypeName<DataType>()));
+			PRAGMA_ENABLE_EXPERIMENTAL_WARNINGS
+		#endif
+
+			return TDataReference(
+				EDataRefShouldConstruct::NewObject,
+				MetasoundDataReferencePrivate::MakeRefOf<DataType, Derived>(Forward<ArgTypes>(Args)...),
+				InDerivedTypeName,
+				InDerivedTypeId
+			);
+		}
+
+
 		/** Enable copy constructor */
 		TDataReference(const TDataReference<DataType>& Other) = default;
 
@@ -337,6 +361,12 @@ namespace Metasound
 
 		template<typename T>
 		friend TDataValueReference<T> ValueCast(const TDataWriteReference<T>& InRef);
+		
+		using FRefType = TDataReference<DataType>::FRefType;
+		TDataValueReference(EDataRefShouldConstruct, FRefType InRef, FName InTypeName, const void* InTypeId)
+			: TDataReference<DataType>(EDataRefShouldConstruct::NewObject, MoveTemp(InRef), InTypeName, InTypeId)
+		{
+		}
 
 	public:
 
@@ -360,6 +390,29 @@ namespace Metasound
 				MetasoundDataReferencePrivate::MakeRefOf<DataType, Derived>(Forward<ArgTypes>(Args)...)
 			);
 		}
+
+		template <typename Derived, typename... ArgTypes>
+		UE_EXPERIMENTAL(5.8, "Polymorphic API is under development")
+		static TDataValueReference CreateNewDerivedAs(const FName& InDerivedTypeName,
+													 FMetasoundDataTypeId InDerivedTypeId,
+													 ArgTypes&&... Args)
+		{
+			static_assert(std::is_base_of_v<DataType, Derived>, "Derived must inherit from DataType");
+			static_assert(std::is_constructible_v<Derived, ArgTypes...>, "Args don't match a Derived constructor");
+
+		#if DO_CHECK
+			PRAGMA_DISABLE_EXPERIMENTAL_WARNINGS
+			check(Metasound::IsChildOfByName(InDerivedTypeName, GetMetasoundDataTypeName<DataType>()));
+			PRAGMA_ENABLE_EXPERIMENTAL_WARNINGS
+		#endif
+			
+			return TDataValueReference(
+				EDataRefShouldConstruct::NewObject,
+				MetasoundDataReferencePrivate::MakeRefOf<DataType, Derived>(Forward<ArgTypes>(Args)...),
+				InDerivedTypeName,
+				InDerivedTypeId
+			);
+		}
 		
 		/** Enable copy constructor */
 		TDataValueReference(const TDataValueReference<DataType>& Other) = default;
@@ -454,6 +507,12 @@ namespace Metasound
 		template <typename T>
 		friend TDataWriteReference<T> WriteCast(const TDataReadReference<T>& InReadableRef);
 
+		using FRefType = TDataReference<DataType>::FRefType;
+		TDataWriteReference(EDataRefShouldConstruct, FRefType InRef, FName InTypeName, const void* InTypeId)
+			: TDataReference<DataType>(EDataRefShouldConstruct::NewObject, MoveTemp(InRef), InTypeName, InTypeId)
+		{
+		}
+
 	public:
 		/** This should be used to construct a new DataType object and return this TDataWriteReference as a wrapper around it. */
 		template <typename... ArgTypes>
@@ -464,7 +523,7 @@ namespace Metasound
 		}
 
 		template <typename Derived, typename... ArgTypes>
-		UE_EXPERIMENTAL(5.7, "Polymorphic API is under development")
+		UE_EXPERIMENTAL(5.8, "Polymorphic API is under development")
 		static TDataWriteReference CreateNewDerived(ArgTypes&&... Args)
 		{
 			static_assert(std::is_base_of_v<DataType, Derived>, "Derived must inherit from DataType");
@@ -475,6 +534,29 @@ namespace Metasound
 				MetasoundDataReferencePrivate::MakeRefOf<DataType, Derived>(Forward<ArgTypes>(Args)...)
 			);
 		}
+
+		template <typename Derived, typename... ArgTypes>
+		UE_EXPERIMENTAL(5.8, "Polymorphic API is under development")
+		static TDataWriteReference CreateNewDerivedAs(const FName& InDerivedTypeName,
+													 FMetasoundDataTypeId InDerivedTypeId,
+													 ArgTypes&&... Args)
+		{
+			static_assert(std::is_base_of_v<DataType, Derived>, "Derived must inherit from DataType");
+			static_assert(std::is_constructible_v<Derived, ArgTypes...>, "Args don't match a Derived constructor");
+
+		#if DO_CHECK
+			PRAGMA_DISABLE_EXPERIMENTAL_WARNINGS
+			check(Metasound::IsChildOfByName(InDerivedTypeName, GetMetasoundDataTypeName<DataType>()));
+			PRAGMA_ENABLE_EXPERIMENTAL_WARNINGS
+		#endif
+			
+			return TDataWriteReference(
+				EDataRefShouldConstruct::NewObject,
+				MetasoundDataReferencePrivate::MakeRefOf<DataType, Derived>(Forward<ArgTypes>(Args)...),
+				InDerivedTypeName,
+				InDerivedTypeId
+			);
+		}
 		
 		/** Enable copy constructor */
 		TDataWriteReference(const TDataWriteReference<DataType>& Other) = default;
@@ -552,6 +634,12 @@ namespace Metasound
 		{
 		}
 
+		using FRefType = typename TDataReference<DataType>::FRefType;
+		TDataReadReference(EDataRefShouldConstruct, FRefType InRef, FName InTypeName, const void* InTypeId)
+			: TDataReference<DataType>(EDataRefShouldConstruct::NewObject, MoveTemp(InRef), InTypeName, InTypeId)
+		{
+		}
+
 	public:
 
 		// This should be used to construct a new DataType object and return this TDataReadReference as a wrapper around it.
@@ -563,7 +651,7 @@ namespace Metasound
 		}
 
 		template <typename Derived, typename... ArgTypes>
-		UE_EXPERIMENTAL(5.7, "Polymorphic API is under development")
+		UE_EXPERIMENTAL(5.8, "Polymorphic API is under development")
 		static TDataReadReference CreateNewDerived(ArgTypes&&... Args)
 		{
 			static_assert(std::is_base_of_v<DataType, Derived>, "Derived must inherit from DataType");
@@ -574,6 +662,27 @@ namespace Metasound
 				MetasoundDataReferencePrivate::MakeRefOf<DataType, Derived>(Forward<ArgTypes>(Args)...)
 			);
 		}
+
+		template <typename Derived, typename... ArgTypes>
+		UE_EXPERIMENTAL(5.8, "Polymorphic API is under development")
+		static TDataReadReference CreateNewDerivedAs(const FName& InDerivedTypeName, FMetasoundDataTypeId InDerivedTypeId, ArgTypes&&... Args)
+		{
+			static_assert(std::is_base_of_v<DataType, Derived>, "Derived must inherit from DataType");
+			static_assert(std::is_constructible_v<Derived, ArgTypes...>, "Args don't match a Derived constructor");
+
+		#if DO_CHECK
+			PRAGMA_DISABLE_EXPERIMENTAL_WARNINGS
+			check(Metasound::IsChildOfByName(InDerivedTypeName, GetMetasoundDataTypeName<DataType>()));
+			PRAGMA_ENABLE_EXPERIMENTAL_WARNINGS
+		#endif
+			
+			return TDataReadReference(
+				EDataRefShouldConstruct::NewObject,
+				MetasoundDataReferencePrivate::MakeRefOf<DataType, Derived>(Forward<ArgTypes>(Args)...),
+				InDerivedTypeName,
+				InDerivedTypeId
+			);
+		}
 		
 		TDataReadReference(const TDataReadReference& Other) = default;
 
diff --git a/Engine/Plugins/Runtime/Metasound/Source/MetasoundGraphCore/Public/MetasoundVertexData.h b/Engine/Plugins/Runtime/Metasound/Source/MetasoundGraphCore/Public/MetasoundVertexData.h
--- a/Engine/Plugins/Runtime/Metasound/Source/MetasoundGraphCore/Public/MetasoundVertexData.h
+++ b/Engine/Plugins/Runtime/Metasound/Source/MetasoundGraphCore/Public/MetasoundVertexData.h
@@ -382,7 +382,8 @@ namespace Metasound
 			void Set(const DataReferenceType& InDataReference)
 			{
 				PRAGMA_DISABLE_INTERNAL_WARNINGS
-				check(Metasound::IsCastable( Vertex.DataTypeName, InDataReference.GetDataTypeName()));
+				checkf(Metasound::IsCastable(InDataReference.GetDataTypeName(), Vertex.DataTypeName),
+					TEXT("Ref=%s, Vertex=%s"), *InDataReference.GetDataTypeName().ToString(), *Vertex.DataTypeName.ToString());
 				PRAGMA_ENABLE_INTERNAL_WARNINGS				
 				Data.Emplace(InDataReference);
 				CheckAccessTypeCompatibility(Vertex, *Data);
diff --git a/Engine/Source/Runtime/SignalProcessing/Private/SphericalHarmonicCalculator.cpp b/Engine/Source/Runtime/SignalProcessing/Private/SphericalHarmonicCalculator.cpp
--- a/Engine/Source/Runtime/SignalProcessing/Private/SphericalHarmonicCalculator.cpp
+++ b/Engine/Source/Runtime/SignalProcessing/Private/SphericalHarmonicCalculator.cpp
@@ -753,6 +753,20 @@ void FSphericalHarmonicCalculator::GenerateFirstOrderRotationMatrixGivenDegrees(
 	GenerateFirstOrderRotationMatrixGivenRadians(RotXDegrees * DEG_2_RAD, RotYDegrees * DEG_2_RAD, RotZDegrees * DEG_2_RAD, OutMatrix);
 }
 
+bool FSphericalHarmonicCalculator::IndexToOrderAndDegree(const int32 InChannelIndex, int32& OutOrder, int32& OutDegree)
+{
+	if (InChannelIndex >= 0)
+	{
+		OutOrder = FMath::FloorToInt(FMath::Sqrt(static_cast<float>(InChannelIndex)));
+		for (; OutOrder * (OutOrder + 1) > InChannelIndex; --OutOrder)
+		{
+		}
+		OutDegree = InChannelIndex - OutOrder * (OutOrder + 1) - OutOrder;
+		return true;
+	}
+	return false;
+}
+
 void FSphericalHarmonicCalculator::NormalizeGains(TArrayView<float> Gains)
 {
 	// Normalize. gains.
diff --git a/Engine/Source/Runtime/SignalProcessing/Private/Vbap.cpp b/Engine/Source/Runtime/SignalProcessing/Private/Vbap.cpp
--- a/Engine/Source/Runtime/SignalProcessing/Private/Vbap.cpp
+++ b/Engine/Source/Runtime/SignalProcessing/Private/Vbap.cpp
@@ -235,12 +235,12 @@ namespace Audio
 			else 
 			{
 				// Don't warn about the wrap span if it's too wide.
-				UE_CLOG(NextIndex != 0, LogTemp, Warning, TEXT("Speaker A=%d,Az=%.2f and Speaker B=%d,Az=%.2f has Width %.2f"),
-					Speaker.Speaker.ChannelID,
-					Speaker.Speaker.AzimuthDegrees,
-					NextSpeaker.Speaker.ChannelID,
-					NextSpeaker.Speaker.AzimuthDegrees,
-					SpanWidth );
+				//UE_CLOG(NextIndex != 0, LogTemp, Warning, TEXT("Speaker A=%d,Az=%.2f and Speaker B=%d,Az=%.2f has Width %.2f"),
+				//	Speaker.Speaker.ChannelID,
+				//	Speaker.Speaker.AzimuthDegrees,
+				//	NextSpeaker.Speaker.ChannelID,
+				//	NextSpeaker.Speaker.AzimuthDegrees,
+				//	SpanWidth );
 			}
 		}
 			
