[Backout] - CL48069941
#rnx
Original CL Desc
-----------------------------------------------------------------
[Experimental] [Metasounds] Polymorphic CAT Update and Granulator node

#tests LowLevelTests

diff --git a/Engine/Plugins/Experimental/MetasoundExperimental/Source/AudioExperimentalRuntime/Private/ChannelAgnostic/ChannelAgnosticTranscoding.cpp b/Engine/Plugins/Experimental/MetasoundExperimental/Source/AudioExperimentalRuntime/Private/ChannelAgnostic/ChannelAgnosticTranscoding.cpp
--- a/Engine/Plugins/Experimental/MetasoundExperimental/Source/AudioExperimentalRuntime/Private/ChannelAgnostic/ChannelAgnosticTranscoding.cpp
+++ b/Engine/Plugins/Experimental/MetasoundExperimental/Source/AudioExperimentalRuntime/Private/ChannelAgnostic/ChannelAgnosticTranscoding.cpp
@@ -6,7 +6,6 @@
 #include "DSP/Ambisonics.h"
 #include "DSP/MultiMono.h"
 #include "DSP/SphericalHarmonicCalculator.h"
-#include "DSP/Vbap.h"
 #include "TypeFamily/ChannelTypeFamily.h"
 
 namespace Audio
@@ -59,47 +58,6 @@ namespace Audio
 			}
 		}
 	}
-
-
-
-	static const IDiscretePanner* ChoosePanner(const FDiscreteChannelTypeFamily& InFrom, const FDiscreteChannelTypeFamily& InTo)
-	{
-		// Return To panner.
-		// All using 2d vbap for now, so just return that.
-		return InTo.GetPanner();
-	}
-	
-	// Static speakers can generate a matrix of speaker gains and cache that to avoid running the panner everytime.
-	bool MakeGainMatrixFrom(const FDiscreteChannelTypeFamily& InSrc, const FDiscreteChannelTypeFamily& InDst, TArray<float>& OutMatrix)
-	{
-		if (const IDiscretePanner* Panner = ChoosePanner(InSrc, InDst))
-		{
-			const int32 DstChannels = InDst.NumChannels();
-			const int32 SrcChannels = InSrc.NumChannels();
-			OutMatrix.SetNum(SrcChannels*DstChannels);
-
-			const TArray<FDiscreteChannelTypeFamily::FSpeaker>& Order = InSrc.GetSpeakerOrder();
-			
-			// For each src channel, create row vector of gains for each dst channel.
-			TArray<float> SpeakerGains;
-			SpeakerGains.SetNum(DstChannels);
-			for (int32 i = 0; i < SrcChannels; ++i)
-			{
-				IDiscretePanner::FOutputParams Results;
-				IDiscretePanner::FInputParams Params { .AzimuthDegrees = Order[i].AzimuthDegrees, .ElevationDegrees = Order[i].ElevationDegrees };
-				if (Panner->ComputeGains(Params,Results) )
-				{
-					// Copy results into gain array.
-					for (int32 j = 0; j < Results.Results.Num(); ++j)
-					{
-						OutMatrix[i*DstChannels+j] = Results.Results[j].Gain;
-					}
-				}
-			}
-			return true;
-		}
-		return false;
-	}
 	
 	// Discrete to Discrete
 	FChannelTypeFamily::FTranscoder GetTranscoder(
@@ -143,16 +101,7 @@ namespace Audio
 
 			case EChannelTranscodeMethod::MixUpOrDown:
 			{
-#if 1 
-				if (TArray<float> Gains; MakeGainMatrixFrom(InSrcType, InDstType, Gains))
-				{
-					return [MixGains = MoveTemp(Gains)](TArrayView<const float*> InSrc, TArrayView<float*> InDst, const int32 NumFrames)
-					{
-						MultiMonoMixUpOrDown(InSrc, InDst, NumFrames, MixGains);
-					};					
-				}
-#else
-					// Make a mix matrix and call a mix/up down.
+				// Make a mix matrix and call a mix/up down.
 				if (TArray<float> Gains; Create2DChannelMap(
 					{
 						.NumInputChannels = InSrcType.NumChannels(),
@@ -171,7 +120,6 @@ namespace Audio
 						MultiMonoMixUpOrDown(InSrc, InDst, NumFrames, MixGains);
 					};
 				}
-#endif //					
 			}
 		}
 		// fail.
@@ -180,7 +128,7 @@ namespace Audio
 
 	// Discrete to Ambisonics
 	FChannelTypeFamily::FTranscoder GetTranscoder(const FDiscreteChannelTypeFamily& InFromType,
-		const FSoundfieldChannelTypeFamily& InToType, const FChannelTypeFamily::FGetTranscoderParams&)
+		const FAmbisonicsChannelTypeFamily& InToType, const FChannelTypeFamily::FGetTranscoderParams&)
 	{
 		const TArray<FDiscreteChannelTypeFamily::FSpeaker>& Speakers = InFromType.GetSpeakerOrder();
 		TArray<TArray<float>> AllChannelGainArrays;
@@ -222,7 +170,7 @@ namespace Audio
 	}
 	
 	// Ambisonics to Discrete
-	FChannelTypeFamily::FTranscoder GetTranscoder(const FSoundfieldChannelTypeFamily& InFromType, const FDiscreteChannelTypeFamily& InToType, const FChannelTypeFamily::FGetTranscoderParams&)
+	FChannelTypeFamily::FTranscoder GetTranscoder(const FAmbisonicsChannelTypeFamily& InFromType, const FDiscreteChannelTypeFamily& InToType, const FChannelTypeFamily::FGetTranscoderParams&)
 	{
 		TArray<TArray<float>> AllChannelGainArrays;
 		AllChannelGainArrays.SetNum(InToType.NumChannels());
@@ -242,7 +190,7 @@ namespace Audio
 			
 			//FVector2f Spherical = Order[i].Position.UnitCartesianToSpherical();
 			//FSphericalHarmonicCalculator::AdjustUESphericalCoordinatesForAmbisonics(Spherical);
-			check(GainArray.Num() == FSoundfieldChannelTypeFamily::OrderToNumChannels(InFromType.GetAmbisonicsOrder()));
+			check(GainArray.Num() == FAmbisonicsChannelTypeFamily::OrderToNumChannels(InFromType.GetAmbisonicsOrder()));
 			const float AzimuthRads = FMath::DegreesToRadians(Speaker.AzimuthDegrees);
 			const float ElevationRads = FMath::DegreesToRadians(Speaker.ElevationDegrees);
 			FSphericalHarmonicCalculator::ComputeSoundfieldChannelGains(InFromType.GetAmbisonicsOrder(), AzimuthRads, ElevationRads, GainArray);
@@ -261,8 +209,8 @@ namespace Audio
 	}
 
 	// Ambisonics to Ambisonics
-	FChannelTypeFamily::FTranscoder GetTranscoder(const FSoundfieldChannelTypeFamily& InFromType,
-		const FSoundfieldChannelTypeFamily& InToType, const FChannelTypeFamily::FGetTranscoderParams& InParams)
+	FChannelTypeFamily::FTranscoder GetTranscoder(const FAmbisonicsChannelTypeFamily& InFromType,
+		const FAmbisonicsChannelTypeFamily& InToType, const FChannelTypeFamily::FGetTranscoderParams& InParams)
 	{
 		// Ambisonics to Ambisonics
 		if (&InToType == &InFromType)
diff --git a/Engine/Plugins/Experimental/MetasoundExperimental/Source/AudioExperimentalRuntime/Private/TypeFamily/ChannelTypeFamily.cpp b/Engine/Plugins/Experimental/MetasoundExperimental/Source/AudioExperimentalRuntime/Private/TypeFamily/ChannelTypeFamily.cpp
--- a/Engine/Plugins/Experimental/MetasoundExperimental/Source/AudioExperimentalRuntime/Private/TypeFamily/ChannelTypeFamily.cpp
+++ b/Engine/Plugins/Experimental/MetasoundExperimental/Source/AudioExperimentalRuntime/Private/TypeFamily/ChannelTypeFamily.cpp
@@ -4,12 +4,9 @@
 
 #include "Algo/Transform.h"
 #include "CoreGlobals.h"
-
-#include "HAL/FileManager.h"
 #include "ChannelAgnostic/ChannelAgnosticTranscoding.h"
 #include "Containers/Map.h"
 #include "DSP/SphericalHarmonicCalculator.h"
-#include "DSP/Vbap.h"
 
 namespace Audio
 {
@@ -26,17 +23,17 @@ namespace Audio
 	
 	class FChannelRegistryImpl final : public IChannelTypeRegistry 
 	{
-		TMap<FName, TSharedRef<const FChannelTypeFamily>> Types;
+		TMap<FName, const FChannelTypeFamily*> Types;
 	public:
 		UE_NONCOPYABLE(FChannelRegistryImpl);
 		FChannelRegistryImpl() = default;
 		virtual ~FChannelRegistryImpl() override = default;
 	protected:
-		virtual bool RegisterType(const FName& InUniqueName, TUniquePtr<FTypeFamily> InType) override
+		virtual bool RegisterType(const FName& InUniqueName, TRetainedRef<FTypeFamily> InType) override
 		{
 			static const FName ChannelTypeName = TEXT("Cat");
 			const FTypeFamily* Type = FindTypeInternal(ChannelTypeName);
-			if (Type && !InType->IsA(Type))
+			if (Type && !InType.Get().IsA(Type))
 			{
 				// Not a cat. (maybe a dog?)
 				return false;	
@@ -49,55 +46,36 @@ namespace Audio
 			}
 			
 			// Safe to cast it and add.
-			TSharedRef<FChannelTypeFamily> Ref(static_cast<FChannelTypeFamily*>(InType.Release()));
+			const FChannelTypeFamily* ChannelType = static_cast<const FChannelTypeFamily*>(&InType.Get());
 			
 			// Add to master list of registered types.
 			check(Types.Find(InUniqueName) == nullptr);
-			Types.Emplace(InUniqueName,MoveTemp(Ref));
-			
+			Types.Add(InUniqueName,ChannelType);
+	
 			return true;
 		}	
 		virtual const FTypeFamily* FindTypeInternal(const FName InUniqueName) const override
 		{
-			if (const auto Found = Types.Find(InUniqueName))
+			if (const FChannelTypeFamily* const * Found = Types.Find(InUniqueName))
 			{
-				return &(Found->Get());
+				return *Found;
 			}
 			return nullptr;
 		}
 
-		virtual TArray<TSharedRef<const FChannelTypeFamily>> GetAllChannelFormats() const override
+		virtual TArray<const FChannelTypeFamily*> GetAllChannelFormats() const override
 		{
-			TArray<TSharedRef<const FChannelTypeFamily>> AllFormats;
+			TArray<const FChannelTypeFamily*> AllFormats;
 			Types.GenerateValueArray(AllFormats);
 			return AllFormats;
 		};
 	};
 
-	int32 FSoundfieldChannelTypeFamily::OrderToNumChannels(const int32 InOrder)
+	int32 FAmbisonicsChannelTypeFamily::OrderToNumChannels(const int32 InOrder)
 	{
 		return FSphericalHarmonicCalculator::OrderToNumChannels(InOrder);
 	}
 
-	TOptional<FChannelTypeFamily::FChannelName> FSoundfieldChannelTypeFamily::GetChannelName(const int32 InChannelIndex) const
-	{
-		if (InChannelIndex < 0)
-		{
-			return {};
-		}
-
-		int32 OrderNum=0;
-		int32 DegreeNum=0;
-		FSphericalHarmonicCalculator::IndexToOrderAndDegree(InChannelIndex, OrderNum, DegreeNum);
-		const FName ChannelName = *FString::Printf(TEXT("ACN %d"), InChannelIndex); // *ChannelName.ToString(),	 (<-- Indexed' FName appears to fail for Node config UX, so use Sprintf'd one).ChannelName(TEXT("ACN"),InChannelIndex; 
-		return FChannelName
-		{
-			.Name = ChannelName,  
-			.FriendlyName = FString::Printf(TEXT("ACN=%d, Order=%d, Degree=%d"), InChannelIndex, Order, DegreeNum)
-		};
-	}
-	
-
 	IChannelTypeRegistry& GetChannelRegistry()
 	{
 		static FChannelRegistryImpl Registry;
@@ -138,17 +116,14 @@ namespace Audio
    FDiscreteChannelTypeFamily::FDiscreteChannelTypeFamily(const FName& InUniqueName, FChannelTypeFamily* InParentType, const FString& InFriendlyName, const TArray<FSpeaker>& InOrder, const bool bIsParentsDefault, const bool bIsAbstract)
 		: Super(InUniqueName, GetFamilyTypeName(), InOrder.Num(), InParentType, InFriendlyName, bIsParentsDefault, bIsAbstract)
 		, Order(InOrder)
-		, Panner(MakePanner(InOrder))
 	{
 		checkf(InParentType, TEXT("Type=%s, Has a Null Parent"), *GetName().ToString());
+		auto GetParentNameSafe = [](FChannelTypeFamily* i) -> FString { return i ? i->GetName().ToString() : FString(); };
+		UE_LOG(LogTemp, Display, TEXT("Unique=%s\tNumChannels=%d\tParent=%s\tFriendlyName=%s\tDefault=%s\tOrder=[%s]\tAbstract=%s"),
+			*InUniqueName.ToString(),  InOrder.Num(), *GetParentNameSafe(InParentType), *InFriendlyName, ToCStr(LexToString(bIsParentsDefault)), *ChannelTypeFamilyPrviate::MakePrettyString(Order), ToCStr(LexToString(bIsAbstract)));
 	}
 
-    FDiscreteChannelTypeFamily::~FDiscreteChannelTypeFamily()
-	{
-		// Needed for TUnique deleter, as panner is forward declared in header.
-	}
-
-   TOptional<FChannelTypeFamily::FChannelName> FDiscreteChannelTypeFamily::GetChannelName(const int32 InChannelIndex) const
+	TOptional<FChannelTypeFamily::FChannelName> FDiscreteChannelTypeFamily::GetChannelName(const int32 InChannelIndex) const
 	{
 		check(InChannelIndex >= 0 && InChannelIndex < NumChannels());
 		if (Order.IsValidIndex(InChannelIndex))
@@ -166,60 +141,245 @@ namespace Audio
 		return {};
 	}
 
-	TUniquePtr<IDiscretePanner> FDiscreteChannelTypeFamily::MakePanner(const TArray<FSpeaker>& InOrder)
+	void RegisterChannelLayouts(IChannelTypeRegistry& Registry)
 	{
-		// Just using 2D VBap for now.
-		TArray<IDiscretePanner::FSpeaker> Speakers;
-		for (int32 Channel = 0; Channel < InOrder.Num(); Channel++)
-		{
-			const FSpeaker& Speaker = InOrder[Channel];
+		/**
+		 * Standard layouts.
+		 * This should be defined in .ini ultimately, allowing new, custom formats to be added simply.
+		 **/
+		
+		// Register root type.
+		static FChannelTypeFamily BaseCat(TEXT("Cat"), TEXT("Cat"),  0, nullptr, TEXT("Base Cat"), false, true);
+		Registry.RegisterType(BaseCat.GetName(), BaseCat);
+		
+#define REGISTER_CAT(NAME, PARENT_NAME, FRIENDLY_NAME, LAYOUT, DEFAULT,ABSTRACT)\
+static FDiscreteChannelTypeFamily UE_JOIN(__Type,NAME)(UE_STRINGIZE(NAME), Registry.FindChannel(PARENT_NAME), FRIENDLY_NAME, LAYOUT, DEFAULT, ABSTRACT);\
+Registry.RegisterType(TEXT(UE_STRINGIZE(NAME)),UE_JOIN(__Type,NAME))
 
-			// Add all speakers that are considered 'spatialized'
-			// This excludes channels that not, like LFE and possibly (Centre channel for dialog). 
-			if (Speaker.bIsSpatialized)
-			{
-				IDiscretePanner::FSpeaker PannerSpeaker
-				{
-					.AzimuthDegrees = Speaker.AzimuthDegrees,
-					.ElevationDegrees = Speaker.ElevationDegrees,
-					.ChannelID = static_cast<int32>(Speaker.Speaker),
-					.ChannelIndex = Channel
-				};
-				Speakers.Add(PannerSpeaker);
-			}
-		}
+		using enum ESpeakerShortNames;
+		using FS = FDiscreteChannelTypeFamily::FSpeaker;
+
+		// Top level abstraction. "Discrete".
+		REGISTER_CAT(Discrete,			TEXT("Cat"),		TEXT("Discrete"),		{},										false,					true);		
+
+		// Mono
+		REGISTER_CAT(Mono, TEXT("Discrete"), TEXT("Mono"), {}, false, true);
+		REGISTER_CAT(Mono1Dot0, TEXT("Mono"), TEXT("Mono (1.0)"),
+		             (TArray<FS>
+			             {
+			             { FC, 0.f, 0.f }, // Az:   0°, El:  0°
+			             }), true, false);
+
+		REGISTER_CAT(Mono1Dot1, TEXT("Mono"), TEXT("Mono (1.1)"), (TArray<FS>{
+			             { FC, 0.f, 0.f }, // Az:   0°, El:  0°
+			             { LFE, 0.f, 0.f } // LFE (omni)
+			             }), false, false);
+
+// Stereo
+REGISTER_CAT(Stereo, TEXT("Discrete"), TEXT("Stereo"), {}, false, true);
+REGISTER_CAT(Stereo2Dot0, TEXT("Stereo"), TEXT("Stereo (2.0)"),
+	(TArray<FS>{
+		{ FL, -30.f,   0.f },  // Az: -30°, El:  0°
+		{ FR,  30.f,   0.f }   // Az: +30°, El:  0°
+		}), true, false);
 		
-		if (Speakers.Num() > 0)
-		{
-			return MakeUnique<FVBap2D>(Speakers);
-		}
-		return {};
-	}
+REGISTER_CAT(Stereo2Dot1, TEXT("Stereo"), TEXT("Stereo (2.1)"),
+	(TArray<FS>{
+		{ FL, -30.f, 0.f },		// Az: -30°, El:  0°
+		{ FR, 30.f, 0.f },		// Az: +30°, El:  0°
+		{ LFE, 0.f, 0.f }		// LFE (omni)
+	}),
+	false, false);
+REGISTER_CAT(Stereo3Dot0, TEXT("Stereo"), TEXT("Stereo (3.0)"),
+	(TArray<FS>{
+		{ FL, -30.f,   0.f },  // Az: -30°, El:  0°
+		{ FC,   0.f,   0.f },  // Az:   0°, El:  0°
+		{ FR,  30.f,   0.f }   // Az: +30°, El:  0°
+	}),
+	false, false);
+REGISTER_CAT(Stereo3DOt1, TEXT("Stereo"), TEXT("Stereo (3.1)"),
+	(TArray<FS>{
+		{ FL, -30.f,   0.f },  // Az: -30°, El:  0°
+		{ FC,   0.f,   0.f },  // Az:   0°, El:  0°
+		{ FR,  30.f,   0.f },  // Az: +30°, El:  0°
+		{ LFE,  0.f,   0.f }   // LFE (omni)
+	}),
+	false, false);
 
-	void RegisterChannelLayouts(IChannelTypeRegistry& Registry)
-	{
-		// Register root type.
-		Registry.RegisterType(TEXT("Cat"), MakeUnique<FChannelTypeFamily>(TEXT("Cat"), TEXT("Cat"),  0, nullptr, TEXT("Base Cat"), false, true));
-		FChannelTypeFamily* BaseCat = Registry.FindChannel(TEXT("Cat"));
-		check(BaseCat);
-
-		Registry.RegisterType(
-			TEXT("Cat:Discrete"),
-			MakeUnique<FDiscreteChannelTypeFamily>(
-						TEXT("Cat:Discrete"),
-						BaseCat,
-						TEXT("Discrete Channel Format"),
-						TArray<FDiscreteChannelTypeFamily::FSpeaker>(),
-						false,
-						true)
-			);
-		Registry.RegisterType(
-			TEXT("Cat:Soundfield"),
-			MakeUnique<FSoundfieldChannelTypeFamily>(
-				TEXT("Cat:Soundfield"),
-				0, BaseCat, TEXT("Soundfield Format"),
-				false,
-				true));
+// Quad
+REGISTER_CAT(Quad, TEXT("Discrete"), TEXT("Quad"), {}, false, true);
+REGISTER_CAT(Quad4Dot0Back, TEXT("Quad"), TEXT("Quad Back Speakers (4.0)"),
+	(TArray<FS>{
+		{ FL, -30.f,   0.f },  // Front Left
+		{ FR,  30.f,   0.f },  // Front Right
+		{ BL, -150.f,  0.f },  // Rear  Left
+		{ BR, 150.f,   0.f }   // Rear  Right
+		}),
+	true, false);
+REGISTER_CAT(Quad4Dot0Side, TEXT("Quad"), TEXT("Quad Side Speakers (4.0)"),
+	(TArray<FS>{
+		{ FL, -30.f,   0.f },  // Front Left
+		{ FR,  30.f,   0.f },  // Front Right
+		{ SL, -90.f,   0.f },  // Side  Left
+		{ SR,  90.f,   0.f }   // Side  Right
+		}),
+		false, false);
+REGISTER_CAT(Quad4Dot1, TEXT("Quad"), TEXT("Quad Back Centre LFE (4.1)"),
+	(TArray<FS>{
+		{ FL, -30.f,   0.f },  // Front Left
+		{ FR,  30.f,   0.f },  // Front Right
+		{ BL, -150.f,  0.f },  // Rear  Left
+		{ BR, 150.f,   0.f },  // Rear  Right
+		{ LFE, 0.f,    0.f }   // LFE (omni)
+		}),
+	false, false);
+
+// Surround
+REGISTER_CAT(Surround, TEXT("Discrete"), TEXT("Surround"), {}, false, true);
+REGISTER_CAT(Surround5, TEXT("Surround"), TEXT("Surround (5.X)"), {}, true, true);
+REGISTER_CAT(Surround5Dot0, TEXT("Surround5"), TEXT("Surround (5.0)"),
+    (TArray<FS>{
+        { FL,  -30,  0 },
+        { FR,   30,  0 },
+        { SL,  -90,  0 },
+        { SR,   90,  0 },
+        { FC,    0,  0 },
+        { BL, -135,  0 },
+        { BR,  135,  0 },
+    }),
+    false, false);
+
+REGISTER_CAT(Surround5_1, TEXT("Surround5"), TEXT("Surround (5.1)"),
+    (TArray<FS>{
+        { FL,  -30,  0 },
+        { FR,   30,  0 },
+        { SL,  -90,  0 },
+        { SR,   90,  0 },
+        { FC,    0,  0 },
+        { BL, -135,  0 },
+        { BR,  135,  0 },
+        { LFE,   0,  0 },
+    }),
+    true, false);
+
+REGISTER_CAT(Surround7, TEXT("Surround"), TEXT("Surround (7.X)"),
+    {}, false, true);
+
+REGISTER_CAT(Surround7Dot0, TEXT("Surround7"), TEXT("Surround (7.0)"),
+    (TArray<FS>{
+        { FL,  -30,  0 },
+        { FR,   30,  0 },
+        { SL,  -90,  0 },
+        { SR,   90,  0 },
+        { FC,    0,  0 },
+        { BL, -135,  0 },
+        { BR,  135,  0 },
+    }),
+    false, false);
+
+REGISTER_CAT(Surround7Dot1, TEXT("Surround7"), TEXT("Surround (7.1)"),
+    (TArray<FS>{
+        { FL,  -30,  0 },
+        { FR,   30,  0 },
+        { SL,  -90,  0 },
+        { SR,   90,  0 },
+        { FC,    0,  0 },
+        { LFE,   0,  0 },
+        { BL, -135,  0 },
+        { BR,  135,  0 },
+    }),
+    true, false);
+
+// Atmos beds & heights
+REGISTER_CAT(Atmos, TEXT("Surround7"), TEXT("Dolby Atmos Bed"),
+    {}, false, true);
+
+REGISTER_CAT(Atmos7Dot0Dot2, TEXT("Atmos"), TEXT("Dolby Atmos (7.0.2)"),
+    (TArray<FS>{
+        { FL,  -30.f,   0.f },
+        { FR,   30.f,   0.f },
+        { SL,  -90.f,   0.f },
+        { SR,   90.f,   0.f },
+        { FC,    0.f,   0.f },
+        { BL, -135.f,   0.f },
+        { BR,  135.f,   0.f },
+        // heights
+        { TFL, -35.3f,  30.0f },
+        { TFR,  35.3f,  30.0f },
+    }),
+    false, false);
+
+REGISTER_CAT(Atmos7Dot0Dot4, TEXT("Atmos"), TEXT("Dolby Atmos (7.0.4)"),
+    (TArray<FS>{
+        { FL,  -30.f,   0.f },
+        { FR,   30.f,   0.f },
+        { SL,  -90.f,   0.f },
+        { SR,   90.f,   0.f },
+        { FC,    0.f,   0.f },
+        { BL, -135.f,   0.f },
+        { BR,  135.f,   0.f },
+        // heights
+        { TFL, -35.3f,  30.0f },
+        { TFR,  35.3f,  30.0f },
+        { TBL, -135.0f, 45.0f },
+        { TBR,  135.0f, 45.0f },
+    }),
+    false, false);
+
+REGISTER_CAT(Atmos7Dot1Dot2, TEXT("Atmos"), TEXT("Dolby Atmos (7.1.2)"),
+    (TArray<FS>{
+        { FL,  -30.f,   0.f },
+        { FR,   30.f,   0.f },
+        { SL,  -90.f,   0.f },
+        { SR,   90.f,   0.f },
+        { FC,    0.f,   0.f },
+        { LFE,   0.f,   0.f },
+        { BL, -135.f,   0.f },
+        { BR,  135.f,   0.f },
+        // heights
+        { TFL, -35.3f,  30.0f },
+        { TFR,  35.3f,  30.0f },
+    }),
+    false, false);
+
+REGISTER_CAT(Atmos7Dot1Dot4, TEXT("Atmos"), TEXT("Dolby Atmos (7.1.4)"),
+    (TArray<FS>{
+        { FL,  -30.f,    0.f },
+        { FR,   30.f,    0.f },
+        { SL,  -90.f,    0.f },
+        { SR,   90.f,    0.f },
+        { FC,    0.f,    0.f },
+        { LFE,   0.f,    0.f },
+        { BL, -135.f,    0.f },
+        { BR,  135.f,    0.f },
+        // heights
+        { TFL, -35.3f,   30.0f },
+        { TFR,  35.3f,   30.0f },
+        { TBL, -135.0f,  45.0f },
+        { TBR,  135.0f,  45.0f },
+    }),
+    true, false);
+
+					
+#define REGISTER_AMBISONICS(NAME, PARENT_NAME, FRIENDLY_NAME, ORDER, DEFAULT,ABSTRACT)\
+static FAmbisonicsChannelTypeFamily UE_JOIN(__Type,NAME)(UE_STRINGIZE(NAME), ORDER, Registry.FindChannel(PARENT_NAME), FRIENDLY_NAME, DEFAULT, ABSTRACT);\
+ensure(Registry.RegisterType(UE_STRINGIZE(NAME),UE_JOIN(__Type,NAME)))
+
+		// Ambisonics 
+		REGISTER_AMBISONICS(Ambisonics,					TEXT("Cat"), 	   TEXT("Ambisonics"), 				0, false, true);
+		REGISTER_AMBISONICS(Ambisonics1stOrder,			TEXT("Ambisonics"), TEXT("1st Order Ambisonics"), 		1, false, false); 
+		REGISTER_AMBISONICS(Ambisonics2ndOrder,			TEXT("Ambisonics"), TEXT("2nd Order Ambisonics"), 		2, false, false); 
+		REGISTER_AMBISONICS(Ambisonics3rdOrder,			TEXT("Ambisonics"), TEXT("3rd Order Ambisonics"), 		3, false, false);
+		REGISTER_AMBISONICS(Ambisonics4thOrder,			TEXT("Ambisonics"), TEXT("4th Order Ambisonics"),		4, false, false); 
+		REGISTER_AMBISONICS(Ambisonics5thOrder,			TEXT("Ambisonics"), TEXT("5th Order Ambisonics"),		5, false, false);  
+		REGISTER_AMBISONICS(Ambisonics6thOrder,			TEXT("Ambisonics"), TEXT("6th Order Ambisonics"),		6, false, false);  
+		REGISTER_AMBISONICS(Ambisonics7thOrder,			TEXT("Ambisonics"), TEXT("7th Order Ambisonics"),		7, false, false);  
+
+		// Channel packs.
+		// FPackedChannelTypeFamily MonoPack( TEXT("PackedChannel")
+		
+#undef REGISTER_CAT
+#undef REGISTER_AMBISONICS
 	}
 }
 
diff --git a/Engine/Plugins/Experimental/MetasoundExperimental/Source/AudioExperimentalRuntime/Public/ChannelAgnostic/ChannelAgnosticTranscoding.h b/Engine/Plugins/Experimental/MetasoundExperimental/Source/AudioExperimentalRuntime/Public/ChannelAgnostic/ChannelAgnosticTranscoding.h
--- a/Engine/Plugins/Experimental/MetasoundExperimental/Source/AudioExperimentalRuntime/Public/ChannelAgnostic/ChannelAgnosticTranscoding.h
+++ b/Engine/Plugins/Experimental/MetasoundExperimental/Source/AudioExperimentalRuntime/Public/ChannelAgnostic/ChannelAgnosticTranscoding.h
@@ -8,9 +8,9 @@ namespace Audio
 {
 	// All transcoder (from -> to) combinations.
 	AUDIOEXPERIMENTALRUNTIME_API FChannelTypeFamily::FTranscoder GetTranscoder(const FDiscreteChannelTypeFamily& InFrom, const FDiscreteChannelTypeFamily& InTo, const FChannelTypeFamily::FGetTranscoderParams&);
-	AUDIOEXPERIMENTALRUNTIME_API FChannelTypeFamily::FTranscoder GetTranscoder(const FDiscreteChannelTypeFamily& InFrom, const FSoundfieldChannelTypeFamily& InTo, const FChannelTypeFamily::FGetTranscoderParams&);
-	AUDIOEXPERIMENTALRUNTIME_API FChannelTypeFamily::FTranscoder GetTranscoder(const FSoundfieldChannelTypeFamily& InFrom, const FDiscreteChannelTypeFamily& InTo, const FChannelTypeFamily::FGetTranscoderParams&);
-	AUDIOEXPERIMENTALRUNTIME_API FChannelTypeFamily::FTranscoder GetTranscoder(const FSoundfieldChannelTypeFamily& InFrom, const FSoundfieldChannelTypeFamily& InTo, const FChannelTypeFamily::FGetTranscoderParams&);
+	AUDIOEXPERIMENTALRUNTIME_API FChannelTypeFamily::FTranscoder GetTranscoder(const FDiscreteChannelTypeFamily& InFrom, const FAmbisonicsChannelTypeFamily& InTo, const FChannelTypeFamily::FGetTranscoderParams&);
+	AUDIOEXPERIMENTALRUNTIME_API FChannelTypeFamily::FTranscoder GetTranscoder(const FAmbisonicsChannelTypeFamily& InFrom, const FDiscreteChannelTypeFamily& InTo, const FChannelTypeFamily::FGetTranscoderParams&);
+	AUDIOEXPERIMENTALRUNTIME_API FChannelTypeFamily::FTranscoder GetTranscoder(const FAmbisonicsChannelTypeFamily& InFrom, const FAmbisonicsChannelTypeFamily& InTo, const FChannelTypeFamily::FGetTranscoderParams&);
 		
 	template<typename TFromType>
 	struct TFromVisitor final : IChannelTypeVisitor
@@ -27,7 +27,7 @@ namespace Audio
 		{
 			Transcoder = GetTranscoder(FromType, InTo, Params);	
 		}
-		virtual void Visit(const FSoundfieldChannelTypeFamily& InTo) override
+		virtual void Visit(const FAmbisonicsChannelTypeFamily& InTo) override
 		{
 			Transcoder = GetTranscoder(FromType, InTo, Params);
 		}
@@ -41,7 +41,7 @@ namespace Audio
 			TFromVisitor Visitor(InFrom, Params, Result);
 			Params.ToType.Accept(Visitor);
 		}
-		virtual void Visit(const FSoundfieldChannelTypeFamily& InFrom) override
+		virtual void Visit(const FAmbisonicsChannelTypeFamily& InFrom) override
 		{
 			TFromVisitor Visitor(InFrom, Params, Result);
 			Params.ToType.Accept(Visitor);
diff --git a/Engine/Plugins/Experimental/MetasoundExperimental/Source/AudioExperimentalRuntime/Public/TypeFamily/ChannelTypeFamily.h b/Engine/Plugins/Experimental/MetasoundExperimental/Source/AudioExperimentalRuntime/Public/TypeFamily/ChannelTypeFamily.h
--- a/Engine/Plugins/Experimental/MetasoundExperimental/Source/AudioExperimentalRuntime/Public/TypeFamily/ChannelTypeFamily.h
+++ b/Engine/Plugins/Experimental/MetasoundExperimental/Source/AudioExperimentalRuntime/Public/TypeFamily/ChannelTypeFamily.h
@@ -9,11 +9,6 @@
 #include "Misc/Optional.h"
 
 
-namespace Audio
-{
-	struct IDiscretePanner;
-}
-
 /**
  * A Comprehensive Short form Speaker enumeration
  */
@@ -36,7 +31,6 @@ enum class ESpeakerShortNames : uint32
 	TBR = EAudioMixerChannel::TopBackRight,			// Top Back Right
 
 	NumChannels = EAudioMixerChannel::ChannelTypeCount,
-	DefaultChannel = EAudioMixerChannel::DefaultChannel,
 };
 
 // Pasted from UHT generated.h, but we can't use because of AudioUnitTests.
@@ -57,7 +51,7 @@ enum class ESpeakerShortNames : uint32
 	op(ESpeakerShortNames::TBL) \
 	op(ESpeakerShortNames::TBR)
 
-AUDIOEXPERIMENTALRUNTIME_API const TCHAR* LexToString(const ESpeakerShortNames InSpeaker);
+const TCHAR* LexToString(const ESpeakerShortNames InSpeaker);
 
 namespace Audio
 {
@@ -69,7 +63,7 @@ namespace Audio
 
 	// Forward declare these for visitor.
 	class FDiscreteChannelTypeFamily;
-	class FSoundfieldChannelTypeFamily;
+	class FAmbisonicsChannelTypeFamily;
 	class FPackedChannelTypeFamily;
 	
 	class IChannelTypeVisitor
@@ -77,7 +71,7 @@ namespace Audio
 	public:
 		virtual ~IChannelTypeVisitor() = default;
 		virtual void Visit(const FDiscreteChannelTypeFamily&) = 0;
-		virtual void Visit(const FSoundfieldChannelTypeFamily&) = 0;
+		virtual void Visit(const FAmbisonicsChannelTypeFamily&) = 0;
 	};
 	
 	// Base type for all "channel" based types.
@@ -159,16 +153,6 @@ namespace Audio
 		 */
 		AUDIOEXPERIMENTALRUNTIME_API FTranscoder GetTranscoder(const FGetTranscoderParams& InParam) const;
 
-		using IPanner = IDiscretePanner;
-		
-		struct FGetPannerParams
-		{
-			// Control which algorithm to use for this format.
-			// Not currently used.
-			// FName PannerType = TEXT("VBAP 2D");
-		};
-		AUDIOEXPERIMENTALRUNTIME_API virtual const IPanner* GetPanner(const FGetPannerParams& InParams = {}) const { return nullptr; }
-		
 		/**
 		 * Returns the family name (C++ type) of this ChannelType. This allows us to safely downcast (if necessary).
 		 * @return The name of this type Family. (i.e. Discrete/Ambisoncs). 
@@ -185,7 +169,7 @@ namespace Audio
 			FName Name;
 			FString FriendlyName;
 		};
-		AUDIOEXPERIMENTALRUNTIME_API virtual TOptional<FChannelName> GetChannelName(const int32 InChannelIndex) const { return {}; }
+		virtual TOptional<FChannelName> GetChannelName(const int32 InChannelIndex) const { return {}; }
 
 		template<typename T>
 		T* Cast()
@@ -208,12 +192,7 @@ namespace Audio
 	protected:
 		// Both are friends so we can hide the api.
 		friend class FDiscreteChannelTypeFamily;
-		friend class FSoundfieldChannelTypeFamily;
-
-		void SetNumChannels(const int32 InNumChannels)
-		{
-			NumChannelsPrivate = InNumChannels;
-		}
+		friend class FAmbisonicsChannelTypeFamily;
 		
 	private:
 		bool bIsAbstract = false;
@@ -223,7 +202,7 @@ namespace Audio
 		FName FamilyType; 
 	};
 	
-	class FDiscreteChannelTypeFamily : public FChannelTypeFamily 
+	class FDiscreteChannelTypeFamily final: public FChannelTypeFamily 
 	{
 	public:
 		UE_NONCOPYABLE(FDiscreteChannelTypeFamily);
@@ -231,10 +210,9 @@ namespace Audio
 
 		struct FSpeaker
 		{
-			ESpeakerShortNames Speaker;				// Enum of speaker types.
-			float AzimuthDegrees = 0.f;				// -180 (left), 0 (front), 180 (right).
-			float ElevationDegrees = 0.f;			// -180 (down), 0 (ground), 180 (up). 
-			bool bIsSpatialized = true;				// If false, azimuth+elevation are ignored.
+			ESpeakerShortNames Speaker;		// Enum of speaker types.
+			float AzimuthDegrees = 0.f;
+			float ElevationDegrees = 0.f;
 		};
 		
 		static FName GetFamilyTypeName()
@@ -243,20 +221,9 @@ namespace Audio
 			return Name;
 		}
 			
-		AUDIOEXPERIMENTALRUNTIME_API FDiscreteChannelTypeFamily(const FName& InUniqueName, FChannelTypeFamily* InParentType, const FString& InFriendlyName,
+		FDiscreteChannelTypeFamily(const FName& InUniqueName, FChannelTypeFamily* InParentType, const FString& InFriendlyName,
 			const TArray<FSpeaker>& InOrder, const bool bIsParentsDefault, const bool bIsAbstract);
 
-		AUDIOEXPERIMENTALRUNTIME_API virtual ~FDiscreteChannelTypeFamily() override;
-		
-		/**
-		 * Returns true if this format defines any height data. 
-		 * @return true if heights are found.
-		 */
-		bool HasHeightChannels() const
-		{
-			return Order.FindByPredicate([](const FSpeaker& i) -> bool { return !FMath::IsNearlyZero(i.ElevationDegrees); }) != nullptr;
-		}
-		
 		/**
 		 * Find the index in the list of Channels for speaker.
 		 * @param InSpeaker 
@@ -283,12 +250,14 @@ namespace Audio
 		 * @return The Function Object. (this will be not be set if a function wasn't returned)
 		 */	
 
-		AUDIOEXPERIMENTALRUNTIME_API virtual TOptional<FChannelName> GetChannelName(const int32 InChannelIndex) const override;
+		virtual TOptional<FChannelName> GetChannelName(const int32 InChannelIndex) const override;
 
 		/**
 		 * Get the order of the speakers defined on this discrete format.
 		 * @return Array of speakers in order.
 		 */
+		//const TArray<ESpeakerShortNames>& GetSpeakerOrder() const { return Order; }
+
 		const TArray<FSpeaker>& GetSpeakerOrder() const { return Order; }
 
 		virtual void Accept(IChannelTypeVisitor& InVisitor) const override
@@ -296,40 +265,31 @@ namespace Audio
 			InVisitor.Visit(*this);
 		}
 
-		virtual const IPanner* GetPanner(const FGetPannerParams& InParams ={}) const override
-		{
-			return Panner.Get();	
-		}
-
 	protected:
-		static TUniquePtr<IDiscretePanner> MakePanner(const TArray<FSpeaker>& Order);
 		TArray<FSpeaker> Order;
-		TUniquePtr<IPanner> Panner;	// Just one for now.
 	};
 
-	// Soundfield 
-	class FSoundfieldChannelTypeFamily : public FChannelTypeFamily 
+	// Ambisonics.
+	class FAmbisonicsChannelTypeFamily final : public FChannelTypeFamily 
 	{
 	public:
-		UE_NONCOPYABLE(FSoundfieldChannelTypeFamily);
+		UE_NONCOPYABLE(FAmbisonicsChannelTypeFamily);
 		using Super = FChannelTypeFamily;
 		
 		static FName GetFamilyTypeName()
 		{
-			static const FName Name = TEXT("Soundfield");
+			static const FName Name = TEXT("Ambisonics");
 			return Name;
 		}
 
 		static int32 OrderToNumChannels(const int32 InOrder);
 
-		AUDIOEXPERIMENTALRUNTIME_API explicit FSoundfieldChannelTypeFamily(const FName& InUniqueName, const int32 InOrder, FChannelTypeFamily* InParentType, const FString& InFriendlyName, const bool bIsParentsDefault, const bool InbIsAbstract)
+		explicit FAmbisonicsChannelTypeFamily(const FName& InUniqueName, const int32 InOrder, FChannelTypeFamily* InParentType, const FString& InFriendlyName, const bool bIsParentsDefault, const bool InbIsAbstract)
 			: Super(InUniqueName, GetFamilyTypeName(), OrderToNumChannels(InOrder), InParentType, InFriendlyName, bIsParentsDefault, InbIsAbstract)
 			, Order(InOrder)
 		{}
 
 		int32 GetAmbisonicsOrder() const { return Order; }
-		
-		AUDIOEXPERIMENTALRUNTIME_API virtual TOptional<FChannelName> GetChannelName(const int32 InChannelIndex) const override;
 
 		virtual void Accept(IChannelTypeVisitor& InVisitor) const override
 		{
@@ -339,24 +299,24 @@ namespace Audio
 		int32 Order = 0;
 	};
 
-	
-	class FCompositeChannelTypeFamily : public FChannelTypeFamily
+	/*
+	class FPackedChannelTypeFamily final : public FChannelTypeFamily
 	{
 	public:
-		UE_NONCOPYABLE(FCompositeChannelTypeFamily);
+		UE_NONCOPYABLE(FPackedChannelTypeFamily);
 
 		using Super = FChannelTypeFamily;
 
 		static FName GetFamilyTypeName()
 		{
-			static const FName Name = TEXT("Composite");
+			static const FName Name = TEXT("Packed");
 			return Name;
 		}
-		AUDIOEXPERIMENTALRUNTIME_API explicit FCompositeChannelTypeFamily(const FName& InUniqueName, FChannelTypeFamily* InParentType, const FString& InFriendlyName, const bool bIsParentsDefault)
+		explicit FPackedChannelTypeFamily(const FName& InUniqueName, FChannelTypeFamily* InParentType, const FString& InFriendlyName, const bool bIsParentsDefault)
 			: Super(InUniqueName, GetFamilyTypeName(), 0, InParentType, InFriendlyName, bIsParentsDefault, false)
 		{}
 	};
-	
+	*/
 	
 	// Channel type registry.
 	class IChannelTypeRegistry : public ITypeFamilyRegistry
@@ -380,8 +340,6 @@ namespace Audio
 		 */
 		const FChannelTypeFamily* FindConcreteChannel(const FName InName) const
 		{
-			check(!InName.IsNone());
-
 			// Walk looking for non-abstract default child. (e.g. Surround->Surround5X->Surround 5.1)
 			const FChannelTypeFamily* Channel = FindChannel(InName);
 			while (Channel && Channel->IsAbstract())
@@ -390,18 +348,12 @@ namespace Audio
 			}
 			return Channel;
 		}
-		const FChannelTypeFamily& FindConcreteChannelChecked(const FName InName) const
-		{
-			const FChannelTypeFamily* Concrete = FindConcreteChannel(InName);
-			checkf(Concrete, TEXT("Concrete Channel not Found: %s"), *InName.ToString());
-			return *Concrete;
-		}
 
 		/**
 		 * Returns every registered format as an array.
 		 * @return Array of format types.
 		 */
-		virtual TArray<TSharedRef<const FChannelTypeFamily>> GetAllChannelFormats() const = 0;
+		virtual TArray<const FChannelTypeFamily*> GetAllChannelFormats() const = 0;
 	};
 
 	// Loose functions for now.
diff --git a/Engine/Plugins/Experimental/MetasoundExperimental/Source/AudioExperimentalRuntime/Public/TypeFamily/TypeFamily.h b/Engine/Plugins/Experimental/MetasoundExperimental/Source/AudioExperimentalRuntime/Public/TypeFamily/TypeFamily.h
--- a/Engine/Plugins/Experimental/MetasoundExperimental/Source/AudioExperimentalRuntime/Public/TypeFamily/TypeFamily.h
+++ b/Engine/Plugins/Experimental/MetasoundExperimental/Source/AudioExperimentalRuntime/Public/TypeFamily/TypeFamily.h
@@ -3,7 +3,6 @@
 #pragma once
 
 #include "Templates/RetainedRef.h"
-#include "Templates/UniquePtr.h"
 #include "UObject/NameTypes.h"
 
 namespace Audio
@@ -38,7 +37,6 @@ namespace Audio
 		bool IsA(const FName& Other) const { return nullptr != FindAncestor([&](const FTypeFamily* i) -> bool { return i->Name==Other; }); }
 		FName GetName() const { return Name; }
 		FString GetFriendlyName() const { return FriendlyName; }
-		const FTypeFamily* GetParent() const { return Parent; }
 	};
 
 	// Simple type registry.
@@ -46,7 +44,7 @@ namespace Audio
 	{
 	public:
 		virtual ~ITypeFamilyRegistry() = default;
-		virtual bool RegisterType(const FName& InUniqueName, const TUniquePtr<FTypeFamily> InType) = 0;
+		virtual bool RegisterType(const FName& InUniqueName, const TRetainedRef<FTypeFamily> InType) = 0;
 
 		template<typename T> const T* Find(const FName& InUniqueName) const
 		{
diff --git a/Engine/Plugins/Experimental/MetasoundExperimental/Source/MetasoundExperimentalEditor/MetasoundExperimentalEditor.build.cs b/Engine/Plugins/Experimental/MetasoundExperimental/Source/MetasoundExperimentalEditor/MetasoundExperimentalEditor.build.cs
--- a/Engine/Plugins/Experimental/MetasoundExperimental/Source/MetasoundExperimentalEditor/MetasoundExperimentalEditor.build.cs
+++ b/Engine/Plugins/Experimental/MetasoundExperimental/Source/MetasoundExperimentalEditor/MetasoundExperimentalEditor.build.cs
@@ -47,8 +47,7 @@ public class MetasoundExperimentalEditor : ModuleRules
 				"MetasoundEditor",
 				"StructUtilsEditor",
 				"MetasoundEngine",
-				"MetasoundFrontend",
-				"MetasoundGraphCore",
+				"MetasoundFrontend"
 				// ... add private dependencies that you statically link with here ...	
 			}
 		);
diff --git a/Engine/Plugins/Experimental/MetasoundExperimental/Source/MetasoundExperimentalEditor/Private/MetasoundExperimentalEditorModule.cpp b/Engine/Plugins/Experimental/MetasoundExperimental/Source/MetasoundExperimentalEditor/Private/MetasoundExperimentalEditorModule.cpp
--- a/Engine/Plugins/Experimental/MetasoundExperimental/Source/MetasoundExperimentalEditor/Private/MetasoundExperimentalEditorModule.cpp
+++ b/Engine/Plugins/Experimental/MetasoundExperimental/Source/MetasoundExperimentalEditor/Private/MetasoundExperimentalEditorModule.cpp
@@ -5,8 +5,6 @@
 #include "MetasoundExampleNodeDetailsCustomization.h"
 #include "MetasoundMappingFunctionNode.h"
 #include "MetasoundMappingFunctionDetailsCustomization.h"
-#include "MetasoundGranulatorNode.h"
-#include "MetasoundGranularNodeDetailsCustomization.h"
 #include "Modules/ModuleManager.h"
 
 #define LOCTEXT_NAMESPACE "MetasoundExperimentalEditor"
@@ -35,12 +33,6 @@ void FMetasoundExperimentalEditorModule::StartupModule()
 		{
 			return MakeShareable(new FMappingFunctionNodeConfigurationCustomization(InStructProperty, InNode));
 		});
-	MetasoundEditorModule.RegisterCustomNodeConfigurationDetailsCustomization(		
-		FMetaSoundGranulatorNodeConfiguration::StaticStruct()->GetFName(),
-		[](TSharedPtr<IPropertyHandle> InStructProperty, TWeakObjectPtr<UMetasoundEditorGraphNode> InNode)
-		{
-			return MakeShareable(new FGranularNodeConfigurationCustomization(InStructProperty, InNode));
-		});
 	MetasoundEditorModule.RegisterPinType("Enum:CurveFunction", "Int32");
 }
 
@@ -49,7 +41,6 @@ void FMetasoundExperimentalEditorModule::ShutdownModule()
 	Metasound::Editor::IMetasoundEditorModule& MetasoundEditorModule = FModuleManager::GetModuleChecked<Metasound::Editor::IMetasoundEditorModule>("MetasoundEditor");
 	MetasoundEditorModule.UnregisterCustomNodeConfigurationDetailsCustomization(FMetaSoundWidgetExampleNodeConfiguration::StaticStruct()->GetFName());
 	MetasoundEditorModule.UnregisterCustomNodeConfigurationDetailsCustomization(FMetaSoundMappingFunctionNodeConfiguration::StaticStruct()->GetFName());
-	MetasoundEditorModule.UnregisterCustomNodeConfigurationDetailsCustomization(FMetaSoundGranulatorNodeConfiguration::StaticStruct()->GetFName());
 }
 
 #undef LOCTEXT_NAMESPACE
diff --git a/Engine/Plugins/Experimental/MetasoundExperimental/Source/MetasoundExperimentalEditor/Private/MetasoundGranularNodeDetailsCustomization.cpp b/Engine/Plugins/Experimental/MetasoundExperimental/Source/MetasoundExperimentalEditor/Private/MetasoundGranularNodeDetailsCustomization.cpp
deleted file mode 100644
diff --git a/Engine/Plugins/Experimental/MetasoundExperimental/Source/MetasoundExperimentalEditor/Private/MetasoundGranularNodeDetailsCustomization.h b/Engine/Plugins/Experimental/MetasoundExperimental/Source/MetasoundExperimentalEditor/Private/MetasoundGranularNodeDetailsCustomization.h
deleted file mode 100644
diff --git a/Engine/Plugins/Experimental/MetasoundExperimental/Source/MetasoundExperimentalEngineRuntime/MetasoundExperimentalEngineRuntime.Build.cs b/Engine/Plugins/Experimental/MetasoundExperimental/Source/MetasoundExperimentalEngineRuntime/MetasoundExperimentalEngineRuntime.Build.cs
--- a/Engine/Plugins/Experimental/MetasoundExperimental/Source/MetasoundExperimentalEngineRuntime/MetasoundExperimentalEngineRuntime.Build.cs
+++ b/Engine/Plugins/Experimental/MetasoundExperimental/Source/MetasoundExperimentalEngineRuntime/MetasoundExperimentalEngineRuntime.Build.cs
@@ -34,8 +34,7 @@ public class MetasoundExperimentalEngineRuntime : ModuleRules
 				"MetasoundFrontend",
 				"MetasoundEngine",
 				"MetasoundStandardNodes",
-				"SignalProcessing",
-				"MetasoundExperimentalRuntime"
+				"SignalProcessing"
 			}
 		);
 
diff --git a/Engine/Plugins/Experimental/MetasoundExperimental/Source/MetasoundExperimentalEngineRuntime/Private/MetasoundGranulatorNode.cpp b/Engine/Plugins/Experimental/MetasoundExperimental/Source/MetasoundExperimentalEngineRuntime/Private/MetasoundGranulatorNode.cpp
deleted file mode 100644
diff --git a/Engine/Plugins/Experimental/MetasoundExperimental/Source/MetasoundExperimentalEngineRuntime/Public/MetasoundGranulatorNode.h b/Engine/Plugins/Experimental/MetasoundExperimental/Source/MetasoundExperimentalEngineRuntime/Public/MetasoundGranulatorNode.h
deleted file mode 100644
diff --git a/Engine/Plugins/Experimental/MetasoundExperimental/Source/MetasoundExperimentalRuntime/MetasoundExperimentalRuntime.Build.cs b/Engine/Plugins/Experimental/MetasoundExperimental/Source/MetasoundExperimentalRuntime/MetasoundExperimentalRuntime.Build.cs
--- a/Engine/Plugins/Experimental/MetasoundExperimental/Source/MetasoundExperimentalRuntime/MetasoundExperimentalRuntime.Build.cs
+++ b/Engine/Plugins/Experimental/MetasoundExperimental/Source/MetasoundExperimentalRuntime/MetasoundExperimentalRuntime.Build.cs
@@ -39,20 +39,15 @@ public class MetasoundExperimentalRuntime : ModuleRules
 				"MetasoundGraphCore",
 				"MetasoundFrontend",
 				"MetasoundStandardNodes",
-				"SignalProcessing",
-				"Serialization",
-				"Projects"
+				"SignalProcessing"
 			}
 		);
+
 		PrivateDefinitions.AddRange(
 			new string[]
 			{
 				"METASOUND_PLUGIN=MetasoundExperimental",
 				"METASOUND_MODULE=MetasoundExperimentalRuntime"
 			});
-		
-		RuntimeDependencies.Add("$(PluginDir)/Content/CatSchemas/Composite/*.json", StagedFileType.UFS);
-		RuntimeDependencies.Add("$(PluginDir)/Content/CatSchemas/Discrete/*.json", StagedFileType.UFS);
-		RuntimeDependencies.Add("$(PluginDir)/Content/CatSchemas/Soundfield/*.json", StagedFileType.UFS);
 	}
 }
diff --git a/Engine/Plugins/Experimental/MetasoundExperimental/Source/MetasoundExperimentalRuntime/Private/MetasoundCatBreakNode.cpp b/Engine/Plugins/Experimental/MetasoundExperimental/Source/MetasoundExperimentalRuntime/Private/MetasoundCatBreakNode.cpp
--- a/Engine/Plugins/Experimental/MetasoundExperimental/Source/MetasoundExperimentalRuntime/Private/MetasoundCatBreakNode.cpp
+++ b/Engine/Plugins/Experimental/MetasoundExperimental/Source/MetasoundExperimentalRuntime/Private/MetasoundCatBreakNode.cpp
@@ -20,42 +20,27 @@ namespace Metasound
 
 	namespace CatBreakNodePrivate
 	{
-		METASOUND_PARAM(InputFromCat,			"Input", "Cat to Break");
+		METASOUND_PARAM(InputFromCat,			"Input", "CAT to Cast");
+		METASOUND_PARAM(OutputToCat,			"Output", "CAT Result");
 
-		static FName GetOutputPinName(const int32 InIndex, const Audio::FChannelTypeFamily& InFormat)
+		const FLazyName InputBaseName{"In"};
+
+		FName MakeOutputVertexName(const int32 InIndex)
 		{
-			if (TOptional<Audio::FChannelTypeFamily::FChannelName> Names = InFormat.GetChannelName(InIndex); Names.IsSet())
-			{
-				return Names->Name;
-			}
-			const FName Output(TEXT("Output"), InIndex);
-			return Output;
+			FName Name = InputBaseName;
+			Name.SetNumber(InIndex);
+			return Name;
 		}
 
-
-		FOutputDataVertex MakeOutputDataVertex(const FName InName, const FString& InFriendlyName)
+		TOutputDataVertex<FAudioBuffer> MakeOutputDataVertex(const int32 InIndex, const FName InName, const FString& InFriendlyName)
 		{
-			const FText DisplayName = METASOUND_LOCTEXT_FORMAT(
-				"DisplayName_{0}", "{0}", FText::FromName(InName));
+			const FName Name = MakeOutputVertexName(InIndex);
 
-			const FText Tooltip = METASOUND_LOCTEXT_FORMAT(
-				"Tooltip_{0}", "Channel {0}", FText::FromString(InFriendlyName));
-			
-			return FOutputDataVertex
-			{
-				InName,
-				GetMetasoundDataTypeName<FAudioBuffer>(),
-				FDataVertexMetadata
-				{
-					.Description = Tooltip,
-					.DisplayName = DisplayName,
-					.bIsAdvancedDisplay = false
-				},
-				EVertexAccessType::Reference
-			};
+			const FText DisplayName = METASOUND_LOCTEXT_FORMAT("In_DisplayName", "{0}", FText::FromString(InName.ToString()));
+
+			return TOutputDataVertex<FAudioBuffer>{Name, FDataVertexMetadata{FText::FromString(*InFriendlyName), DisplayName}};
 		}
 		
-		
 		FVertexInterface MakeClassInterface(const FName Format)
 		{
 			const Audio::FChannelTypeFamily* Found = Audio::GetChannelRegistry().FindConcreteChannel(Format);
@@ -66,19 +51,18 @@ namespace Metasound
 			const int32 NumChannels = Found->NumChannels();
 
 			FInputVertexInterface InputInterface;
-            InputInterface.Add(FInputDataVertex(METASOUND_GET_PARAM_NAME(InputFromCat), Format, METASOUND_GET_PARAM_METADATA(InputFromCat), EVertexAccessType::Reference));
+            InputInterface.Add(TInputDataVertex<FChannelAgnosticType>(METASOUND_GET_PARAM_NAME_AND_METADATA(OutputToCat)));
 			
 			FOutputVertexInterface OutputInterface;
 			for (int32 ChannelIndex = 0; ChannelIndex < NumChannels; ++ChannelIndex )
 			{
 				if (TOptional<Audio::FChannelTypeFamily::FChannelName> Name = Found->GetChannelName(ChannelIndex); Name.IsSet())
 				{
-					OutputInterface.Add(MakeOutputDataVertex(Name->Name, Name->FriendlyName));
+					OutputInterface.Add(MakeOutputDataVertex(ChannelIndex, Name->Name, Name->FriendlyName));
 				}
 				else
 				{
-					const FName ChannelName(TEXT("Output"), ChannelIndex);
-					OutputInterface.Add(MakeOutputDataVertex(ChannelName, FString::Printf(TEXT("Channel %d"), ChannelIndex)));	
+					OutputInterface.Add(MakeOutputDataVertex(ChannelIndex, TEXT("Output"), FString::FromInt(ChannelIndex) ));	
 				}
 			}
 			return{
@@ -145,7 +129,26 @@ namespace Metasound
 			, Format(InFormat)
 		{}
 		virtual ~FCatBreakOperator() override = default;
-		
+
+		static const FVertexInterface& GetDefaultInterface()
+		{
+			using namespace CatBreakNodePrivate ;
+			auto CreateDefaultInterface = []()-> FVertexInterface
+			{
+				// inputs
+				FInputVertexInterface InputInterface;
+				InputInterface.Add(TInputDataVertex<FChannelAgnosticType>(METASOUND_GET_PARAM_NAME_AND_METADATA(InputFromCat)));
+				
+				// outputs
+				FOutputVertexInterface OutputInterface;
+				
+				return FVertexInterface(MoveTemp(InputInterface), MoveTemp(OutputInterface));
+			}; // end lambda: CreateDefaultInterface()
+
+			static const FVertexInterface DefaultInterface = CreateDefaultInterface();
+			return DefaultInterface;
+		}
+			
 		static TUniquePtr<IOperator> CreateOperator(const FBuildOperatorParams& InParams, FBuildResults& OutResults)
 		{
 			using namespace CatBreakNodePrivate;
@@ -167,7 +170,7 @@ namespace Metasound
 
 			// Create the input pin.
 			TDataReadReference<FChannelAgnosticType> InputPin = InputData.GetOrCreateDefaultDataReadReference<FChannelAgnosticType>(
-				METASOUND_GET_PARAM_NAME(InputFromCat), InParams.OperatorSettings);
+				METASOUND_GET_PARAM_NAME(OutputToCat), InParams.OperatorSettings);
 
 			// Create output based on format.
 			TArray<TDataWriteReference<FAudioBuffer>> OutputAudioVertices;
@@ -197,14 +200,14 @@ namespace Metasound
 		virtual void BindInputs(FInputVertexInterfaceData& InOutVertexData) override
 		{
 			using namespace CatBreakNodePrivate;
-			InOutVertexData.BindReadVertex(METASOUND_GET_PARAM_NAME(InputFromCat),InputCat);
+			InOutVertexData.BindReadVertex(METASOUND_GET_PARAM_NAME(OutputToCat),InputCat);
 		}
 		virtual void BindOutputs(FOutputVertexInterfaceData& InOutVertexData) override
 		{
 			using namespace CatBreakNodePrivate;
 			for (int32 i = 0; i < OutputAudioVertices.Num(); i++)
 			{
-				InOutVertexData.BindWriteVertex(GetOutputPinName(i, InputCat->GetType()), OutputAudioVertices[i]);
+				InOutVertexData.BindWriteVertex(MakeOutputVertexName(i), OutputAudioVertices[i]);
 			}
 		}
 
@@ -237,7 +240,7 @@ namespace Metasound
 				METASOUND_LOCTEXT("CatBreakNodeNameDescription", "A Node that Breaks CATs"),	
 				TEXT("UE"), // Author
 				METASOUND_LOCTEXT("ExampleConfigurablePromptIfMissing", "Enable the MetaSoundExperimental Plugin"), // Prompt if missing
-				MakeClassInterface(TEXT("Cat:Stereo2Dot0")),
+				GetDefaultInterface(),
 				{}
 			};
 		}
diff --git a/Engine/Plugins/Experimental/MetasoundExperimental/Source/MetasoundExperimentalRuntime/Private/MetasoundCatBreakNode.h b/Engine/Plugins/Experimental/MetasoundExperimental/Source/MetasoundExperimentalRuntime/Private/MetasoundCatBreakNode.h
--- a/Engine/Plugins/Experimental/MetasoundExperimental/Source/MetasoundExperimentalRuntime/Private/MetasoundCatBreakNode.h
+++ b/Engine/Plugins/Experimental/MetasoundExperimental/Source/MetasoundExperimentalRuntime/Private/MetasoundCatBreakNode.h
@@ -15,8 +15,8 @@ struct FMetaSoundCatBreakNodeConfiguration : public FMetaSoundFrontendNodeConfig
 
 	FMetaSoundCatBreakNodeConfiguration() = default;
 
-	UPROPERTY(EditAnywhere, Category = General, meta = (GetOptions="MetasoundCatCastingOptionsHelper.GetCastingOptions_NoAbstract"))
-	FName Format = TEXT("Cat:Stereo2Dot0");
+	UPROPERTY(EditAnywhere, Category = General, meta = (GetOptions="MetasoundCatCastingOptionsHelper.GetCastingOptions"))
+	FName Format = TEXT("Stereo");
 
 	UPROPERTY(EditAnywhere, Category = General)
 	EMetasoundCatCastingMethod TranscodeMethod = EMetasoundCatCastingMethod::ChannelDrop;
diff --git a/Engine/Plugins/Experimental/MetasoundExperimental/Source/MetasoundExperimentalRuntime/Private/MetasoundCatCastingNode.cpp b/Engine/Plugins/Experimental/MetasoundExperimental/Source/MetasoundExperimentalRuntime/Private/MetasoundCatCastingNode.cpp
--- a/Engine/Plugins/Experimental/MetasoundExperimental/Source/MetasoundExperimentalRuntime/Private/MetasoundCatCastingNode.cpp
+++ b/Engine/Plugins/Experimental/MetasoundExperimental/Source/MetasoundExperimentalRuntime/Private/MetasoundCatCastingNode.cpp
@@ -2,13 +2,11 @@
 
 #include "MetasoundCatCastingNode.h"
 
-#include "MetasoundCatDataSchemas.h"
 #include "MetasoundChannelAgnosticType.h"
 #include "MetasoundExecutableOperator.h"
 #include "MetasoundFacade.h"
 #include "MetasoundNodeRegistrationMacro.h"
 #include "MetasoundVertex.h"
-#include "DSP/MultiMono.h"
 #include "TypeFamily/ChannelTypeFamily.h"
 
 #define LOCTEXT_NAMESPACE "MetasoundStandardNodes_CatCastingNode"
@@ -51,166 +49,144 @@ namespace Metasound
 		// Linkage.
 		const FLazyName FCatCastingOperatorData::OperatorDataTypeName = TEXT("FCatCastingOperatorData");
 	}
-
-	static TDataWriteReference<FChannelAgnosticType>
-	MakePolyCat(const FName InDerivedType, const FOperatorSettings& InOpSettings)
+	
+	class FCatCastingOperator final : public TExecutableOperator<FCatCastingOperator>
 	{
-		using namespace Frontend;
-		const FLiteral Literal(InDerivedType.ToString());
-		TOptional<FAnyDataReference> Ref = IDataTypeRegistry::Get().CreateDataReference(InDerivedType, EDataReferenceAccessType::Write, Literal, InOpSettings);
-		return Ref->GetAs<TDataWriteReference<FChannelAgnosticType>>();
-	}
-
-	FCatCastingOperator::FCatCastingOperator(const FBuildOperatorParams& InParams, FChannelAgnosticTypeReadRef&& InInputCat,
-	                                         const CatCastingPrivate::FCatCastingOperatorData& InData, const FName InConcreteName)
-		: InputFrom(MoveTemp(InInputCat))
-		, ToFormatName(InConcreteName)
-		, OutputCastResult(MakePolyCat(InConcreteName, InParams.OperatorSettings))
-		, Settings(InParams.OperatorSettings)
-		, TranscodeMethod(InData.GetTranscodeMethod())
-		, MixMethod(InData.GetMixMethod())
-	{
-	}
-
-	FVertexInterface FCatCastingOperator::GetInterface(const FName InSrc, const FName InDst) 
-	{
-		using namespace CatCastingPrivate;
-
-		// inputs
-		FInputVertexInterface InputInterface;
-		InputInterface.Add(
-			FInputDataVertex(
-				METASOUND_GET_PARAM_NAME(InputFromCat),
-				InSrc,
-				METASOUND_GET_PARAM_METADATA(InputFromCat)
-			)
-		);
-		
-		// outputs
-		FOutputVertexInterface OutputInterface;
-		OutputInterface.Add(
-			FOutputDataVertex(
-				METASOUND_GET_PARAM_NAME(OutputToCat),
-				InDst,
-				METASOUND_GET_PARAM_METADATA(OutputToCat),
-				EVertexAccessType::Reference
-			)
-		);
-		
-		return FVertexInterface(InputInterface, OutputInterface);
-	}
-
-	TUniquePtr<IOperator> FCatCastingOperator::CreateOperator(const FBuildOperatorParams& InParams, FBuildResults& OutResults)
-	{
-		using namespace CatCastingPrivate;
-
-		const FCatCastingOperatorData* CatTestingConfigData = CastOperatorData<const FCatCastingOperatorData>(InParams.Node.GetOperatorData().Get());
-		if (!CatTestingConfigData)
+	public:
+		FCatCastingOperator(const FBuildOperatorParams& InParams, FChannelAgnosticTypeReadRef&& InInputCat, const CatCastingPrivate::FCatCastingOperatorData& InData, const FName InConcreteName)
+			: InputFrom(MoveTemp(InInputCat))
+			, ToFormatName(InConcreteName)
+			, OutputCastResult(FChannelAgnosticTypeWriteRef::CreateNew(InParams.OperatorSettings, InConcreteName))
+			, Settings(InParams.OperatorSettings)
+			, TranscodeMethod(InData.GetTranscodeMethod())
+			, MixMethod(InData.GetMixMethod())
+		{}
+		virtual ~FCatCastingOperator() override = default;
+
+		static const FVertexInterface& GetDefaultInterface()
 		{
-			return MakeUnique<FNoOpOperator>();
+			using namespace CatCastingPrivate;
+			auto CreateDefaultInterface = []()-> FVertexInterface
+			{
+				// inputs
+				FInputVertexInterface InputInterface;
+				InputInterface.Add(TInputDataVertex<FChannelAgnosticType>(METASOUND_GET_PARAM_NAME_AND_METADATA(InputFromCat)));
+				
+				// outputs
+				FOutputVertexInterface OutputInterface;
+				OutputInterface.Add(TOutputDataVertex<FChannelAgnosticType>(METASOUND_GET_PARAM_NAME_AND_METADATA(OutputToCat)));
+				
+				return FVertexInterface(InputInterface, OutputInterface);
+			}; // end lambda: CreateDefaultInterface()
+
+			static const FVertexInterface DefaultInterface = CreateDefaultInterface();
+			return DefaultInterface;
 		}
-		const FName CastToName = CatTestingConfigData->GetToType();
-
-		const FInputVertexInterfaceData& InputData = InParams.InputData;
-		const Audio::FChannelTypeFamily* ConcreteToType = Audio::GetChannelRegistry().FindConcreteChannel(CastToName);
-
-		TDataReadReference<FChannelAgnosticType> InputCat = InputData.GetOrCreateDefaultDataReadReference<FChannelAgnosticType>(
-			METASOUND_GET_PARAM_NAME(InputFromCat), InParams.OperatorSettings);
-
-		// Make sure the cast is to something sane, otherwise use the inputs type... 
-		const FName CastToNameSane = ConcreteToType ? ConcreteToType->GetName() : InputCat->GetTypeName();
-
-		return MakeUnique<FCatCastingOperator>(
-			InParams,
-			MoveTemp(InputCat),
-			*CatTestingConfigData,
-			CastToNameSane
-		);
-	}
-
-	void FCatCastingOperator::BindInputs(FInputVertexInterfaceData& InOutVertexData) 
-	{
-		using namespace CatCastingPrivate;
-		InOutVertexData.BindReadVertex(METASOUND_GET_PARAM_NAME(InputFromCat), InputFrom);
-	}
-
-	void FCatCastingOperator::BindOutputs(FOutputVertexInterfaceData& InOutVertexData) 
-	{
-		using namespace CatCastingPrivate;
-		InOutVertexData.BindWriteVertex(METASOUND_GET_PARAM_NAME(OutputToCat), OutputCastResult);
+			
+		static TUniquePtr<IOperator> CreateOperator(const FBuildOperatorParams& InParams, FBuildResults& OutResults)
+		{
+			using namespace CatCastingPrivate;
 
-		// Create transcoder 
-		Transcoder = InputFrom->GetType().GetTranscoder(
+			const FCatCastingOperatorData* CatTestingConfigData = CastOperatorData<const FCatCastingOperatorData>(InParams.Node.GetOperatorData().Get());
+			if (!CatTestingConfigData)
 			{
-				.ToType = OutputCastResult->GetType(),
-				.TranscodeMethod = TranscodeMethod,
-				.MixMethod = MixMethod,
+				return MakeUnique<FNoOpOperator>();
 			}
-		);
-	}
-
-	void FCatCastingOperator::Reset(const FResetParams& InParams)
-	{
-		Execute();
-	}
+			const FName CastToName = CatTestingConfigData->GetToType();
+			
+			const FInputVertexInterfaceData& InputData = InParams.InputData;
+			const Audio::FChannelTypeFamily* ConcreteToType = Audio::GetChannelRegistry().FindConcreteChannel(CastToName);
+
+			TDataReadReference<FChannelAgnosticType> InputCat = InputData.GetOrCreateDefaultDataReadReference<FChannelAgnosticType>(
+				METASOUND_GET_PARAM_NAME(InputFromCat), InParams.OperatorSettings);
+
+			// Make sure the cast is to something sane, otherwise use the inputs type... 
+			const FName CastToNameSane = ConcreteToType ? ConcreteToType->GetName() : InputCat->GetTypeName();
+			
+			return MakeUnique<FCatCastingOperator>(
+				InParams,
+				MoveTemp(InputCat),
+				*CatTestingConfigData,
+				CastToNameSane
+			);
+		}
 
-	void FCatCastingOperator::Execute()
-	{
-		if (Transcoder)
+		virtual void BindInputs(FInputVertexInterfaceData& InOutVertexData) override
 		{
 			using namespace CatCastingPrivate;
-			using namespace Audio;
-			TStackArrayOfPointers<const float> Src = MakeMultiMonoPointersFromView(InputFrom->GetRawMultiMono(), Settings.GetNumFramesPerBlock(),
-			                                                                       InputFrom->NumChannels());
-			TStackArrayOfPointers<float> Dst = MakeMultiMonoPointersFromView(OutputCastResult->GetRawMultiMono(), Settings.GetNumFramesPerBlock(),
-			                                                                 OutputCastResult->NumChannels());
-			Transcoder(Src, Dst, Settings.GetNumFramesPerBlock());
+			InOutVertexData.BindReadVertex(METASOUND_GET_PARAM_NAME(InputFromCat) , InputFrom);
+		}
+		
+		virtual void BindOutputs(FOutputVertexInterfaceData& InOutVertexData) override
+		{
+			using namespace CatCastingPrivate;
+			InOutVertexData.BindWriteVertex(METASOUND_GET_PARAM_NAME(OutputToCat),OutputCastResult);
+
+			// Create transcoder 
+			Transcoder = InputFrom->GetType().GetTranscoder(
+				{
+					.ToType = OutputCastResult->GetType(),
+					.TranscodeMethod = TranscodeMethod,
+					.MixMethod = MixMethod,
+				}
+			);
 		}
-	}
 
-	FNodeClassMetadata FCatCastingOperator::GetNodeInfo()
-	{
-		using namespace CatCastingPrivate;
-		static FName DefaultFormat = TEXT("Cat");
-		return FNodeClassMetadata
+		void Reset(const FResetParams& InParams)
+		{
+			Execute();
+		}
+
+		void Execute()
+		{
+			if (Transcoder)
+			{
+				using namespace CatCastingPrivate;
+				using namespace Audio;
+				TStackArrayOfPointers<const float> Src = MakeMultiMonoPointersFromView(InputFrom->GetRawMultiMono(), Settings.GetNumFramesPerBlock(), InputFrom->NumChannels());
+				TStackArrayOfPointers<float> Dst = MakeMultiMonoPointersFromView(OutputCastResult->GetRawMultiMono(), Settings.GetNumFramesPerBlock(), OutputCastResult->NumChannels());
+				Transcoder(Src, Dst,  Settings.GetNumFramesPerBlock());
+			}
+		}
+		static FNodeClassMetadata GetNodeInfo()
 		{
-			FNodeClassName{"Experimental", "CatCastingOperator", ""},
-			1, // Major version
-			0, // Minor version
+			using namespace CatCastingPrivate;
+			return FNodeClassMetadata
+			{
+				FNodeClassName{ "Experimental", "CatCastingOperator", "" },
+				1, // Major version
+				0, // Minor version
 				METASOUND_LOCTEXT("CatCastingNodeName", "CAT Casting Node"),	
 				METASOUND_LOCTEXT("CatCastingNodeNameDescription", "A Node that allows Casting to CATs"),	
-			TEXT("UE"), // Author
+				TEXT("UE"), // Author
 				METASOUND_LOCTEXT("ExampleConfigurablePromptIfMissing", "Enable the MetaSoundExperimental Plugin"), // Prompt if missing
-			GetInterface(DefaultFormat, DefaultFormat),
-			{}
-		};
-	}
+				GetDefaultInterface(),
+				{}
+			};
+		}
+	private:
+		using FTranscoder = Audio::FChannelTypeFamily::FTranscoder;
+		FChannelAgnosticTypeReadRef InputFrom;
+		FName ToFormatName;
+		FChannelAgnosticTypeWriteRef OutputCastResult;
+		FOperatorSettings Settings;
+		FTranscoder Transcoder;
+		Audio::EChannelTranscodeMethod TranscodeMethod;
+		Audio::EChannelMapMonoUpmixMethod MixMethod;
+	}; // class FCatCastingOperator
+		
+	using FCatCastingNode = TNodeFacade<FCatCastingOperator>;
 
 	// register node config.	
 	METASOUND_REGISTER_NODE_AND_CONFIGURATION(FCatCastingNode, FMetaSoundCatCastingNodeConfiguration);
-}; // class FCatCastingOperator
-
-
+} // namespace Metasound
 
 
 TArray<FPropertyTextFName> UMetasoundCatCastingOptionsHelper::GetCastingOptions()
 {
-	TArray<TSharedRef<const Audio::FChannelTypeFamily>> AllFormats = Audio::GetChannelRegistry().GetAllChannelFormats();
-	TArray<FPropertyTextFName> FormatsOptions;
-	Algo::Transform(AllFormats, FormatsOptions, [](const TSharedRef<const Audio::FChannelTypeFamily>& i) -> FPropertyTextFName
-		{ return {.ValueString = i->GetName(), .DisplayName = FText::FromString(i->GetFriendlyName()) }; });
-	return FormatsOptions;
-}
-
-TArray<FPropertyTextFName> UMetasoundCatCastingOptionsHelper::GetCastingOptions_NoAbstract()
-{
-	TArray<TSharedRef<const Audio::FChannelTypeFamily>> AllFormats = Audio::GetChannelRegistry().GetAllChannelFormats();
+	const TArray<const Audio::FChannelTypeFamily*> AllFormats = Audio::GetChannelRegistry().GetAllChannelFormats();
 	TArray<FPropertyTextFName> FormatsOptions;
-	Algo::TransformIf(
-		AllFormats,
-		FormatsOptions,
-		[](const TSharedRef<const Audio::FChannelTypeFamily>& i) -> bool { return !i->IsAbstract(); },
-		[](const TSharedRef<const Audio::FChannelTypeFamily>& i) -> FPropertyTextFName
+	Algo::Transform(AllFormats, FormatsOptions, [](const Audio::FChannelTypeFamily* i) -> FPropertyTextFName
 		{ return {.ValueString = i->GetName(), .DisplayName = FText::FromString(i->GetFriendlyName()) }; });
 	return FormatsOptions;
 }
@@ -218,10 +194,7 @@ TArray<FPropertyTextFName> UMetasoundCatCastingOptionsHelper::GetCastingOptions_
 TInstancedStruct<FMetasoundFrontendClassInterface> FMetaSoundCatCastingNodeConfiguration::OverrideDefaultInterface(
 	const FMetasoundFrontendClass& InNodeClass) const
 {
-	static const FName DefaultInput = TEXT("Cat");
-	return TInstancedStruct<FMetasoundFrontendClassInterface>::Make(
-		FMetasoundFrontendClassInterface::GenerateClassInterface(
-			Metasound::FCatCastingOperator::GetInterface(DefaultInput,  Metasound::AddCatNamespaceToName(ToType))));
+	return  TInstancedStruct<FMetasoundFrontendClassInterface>::Make(FMetasoundFrontendClassInterface::GenerateClassInterface(Metasound::FCatCastingOperator::GetDefaultInterface()));
 }
 
 TSharedPtr<const Metasound::IOperatorData> FMetaSoundCatCastingNodeConfiguration::GetOperatorData() const
diff --git a/Engine/Plugins/Experimental/MetasoundExperimental/Source/MetasoundExperimentalRuntime/Private/MetasoundCatCastingNode.h b/Engine/Plugins/Experimental/MetasoundExperimental/Source/MetasoundExperimentalRuntime/Private/MetasoundCatCastingNode.h
--- a/Engine/Plugins/Experimental/MetasoundExperimental/Source/MetasoundExperimentalRuntime/Private/MetasoundCatCastingNode.h
+++ b/Engine/Plugins/Experimental/MetasoundExperimental/Source/MetasoundExperimentalRuntime/Private/MetasoundCatCastingNode.h
@@ -2,8 +2,6 @@
 
 #pragma once
 
-#include "MetasoundChannelAgnosticType.h"
-#include "MetasoundExecutableOperator.h"
 #include "MetasoundFrontendDocument.h"
 #include "Templates/SharedPointer.h"
 #include "TypeFamily/ChannelTypeFamily.h"
@@ -18,9 +16,6 @@ class UMetasoundCatCastingOptionsHelper : public UObject
 
 	UFUNCTION()
 	static TArray<FPropertyTextFName> GetCastingOptions();
-
-	UFUNCTION()
-	static TArray<FPropertyTextFName> GetCastingOptions_NoAbstract();
 };
 
 UENUM()
@@ -45,7 +40,7 @@ struct FMetaSoundCatCastingNodeConfiguration : public FMetaSoundFrontendNodeConf
 
 	FMetaSoundCatCastingNodeConfiguration() =  default;
 
-	UPROPERTY(EditAnywhere, Category = General, meta = (GetOptions="MetasoundCatCastingOptionsHelper.GetCastingOptions_NoAbstract"))
+	UPROPERTY(EditAnywhere, Category = General, meta = (GetOptions="MetasoundCatCastingOptionsHelper.GetCastingOptions"))
 	FName ToType = TEXT("Mono");
 	
 	UPROPERTY(EditAnywhere, Category = General)
@@ -60,47 +55,3 @@ struct FMetaSoundCatCastingNodeConfiguration : public FMetaSoundFrontendNodeConf
 	/* Pass data down to the operator. */
 	virtual TSharedPtr<const Metasound::IOperatorData> GetOperatorData() const override;
 };
-
-namespace Metasound
-{
-	namespace CatCastingPrivate
-	{
-		class FCatCastingOperatorData;
-	}
-
-	struct FBuildOperatorParams;
-
-	class FCatCastingOperator final : public TExecutableOperator<FCatCastingOperator>
-	{
-	public:
-		FCatCastingOperator(const FBuildOperatorParams& InParams, FChannelAgnosticTypeReadRef&& InInputCat, const CatCastingPrivate::FCatCastingOperatorData& InData, const FName InConcreteName);
-		virtual ~FCatCastingOperator() override = default;
-
-		static FVertexInterface GetInterface(const FName InSrc, const FName InDst);
-			
-		static TUniquePtr<IOperator> CreateOperator(const FBuildOperatorParams& InParams, FBuildResults& OutResults);
-
-		virtual void BindInputs(FInputVertexInterfaceData& InOutVertexData) override;
-		
-		virtual void BindOutputs(FOutputVertexInterfaceData& InOutVertexData) override;
-
-		void Reset(const FResetParams& InParams);
-
-		void Execute();
-	
-		static FNodeClassMetadata GetNodeInfo();
-	
-	private:
-		using FTranscoder = Audio::FChannelTypeFamily::FTranscoder;
-		FChannelAgnosticTypeReadRef InputFrom;
-		FName ToFormatName;
-		TDataWriteReference<FChannelAgnosticType> OutputCastResult;
-		FOperatorSettings Settings;
-		FTranscoder Transcoder;
-		Audio::EChannelTranscodeMethod TranscodeMethod;
-		Audio::EChannelMapMonoUpmixMethod MixMethod;
-	}; // class FCatCastingOperator
-
-	using FCatCastingNode = TNodeFacade<FCatCastingOperator>;
-
-}		
diff --git a/Engine/Plugins/Experimental/MetasoundExperimental/Source/MetasoundExperimentalRuntime/Private/MetasoundCatConversionNode.cpp b/Engine/Plugins/Experimental/MetasoundExperimental/Source/MetasoundExperimentalRuntime/Private/MetasoundCatConversionNode.cpp
deleted file mode 100644
diff --git a/Engine/Plugins/Experimental/MetasoundExperimental/Source/MetasoundExperimentalRuntime/Private/MetasoundCatConversionNode.h b/Engine/Plugins/Experimental/MetasoundExperimental/Source/MetasoundExperimentalRuntime/Private/MetasoundCatConversionNode.h
deleted file mode 100644
diff --git a/Engine/Plugins/Experimental/MetasoundExperimental/Source/MetasoundExperimentalRuntime/Private/MetasoundCatDataSchemas.cpp b/Engine/Plugins/Experimental/MetasoundExperimental/Source/MetasoundExperimentalRuntime/Private/MetasoundCatDataSchemas.cpp
deleted file mode 100644
diff --git a/Engine/Plugins/Experimental/MetasoundExperimental/Source/MetasoundExperimentalRuntime/Private/MetasoundCatDataSchemas.h b/Engine/Plugins/Experimental/MetasoundExperimental/Source/MetasoundExperimentalRuntime/Private/MetasoundCatDataSchemas.h
deleted file mode 100644
diff --git a/Engine/Plugins/Experimental/MetasoundExperimental/Source/MetasoundExperimentalRuntime/Private/MetasoundCatMakeNode.cpp b/Engine/Plugins/Experimental/MetasoundExperimental/Source/MetasoundExperimentalRuntime/Private/MetasoundCatMakeNode.cpp
--- a/Engine/Plugins/Experimental/MetasoundExperimental/Source/MetasoundExperimentalRuntime/Private/MetasoundCatMakeNode.cpp
+++ b/Engine/Plugins/Experimental/MetasoundExperimental/Source/MetasoundExperimentalRuntime/Private/MetasoundCatMakeNode.cpp
@@ -22,28 +22,26 @@ namespace Metasound
 		METASOUND_PARAM(InputFromCat,			"Input", "CAT to Cast");
 		METASOUND_PARAM(OutputToCat,			"Output", "CAT Result");
 
-		FInputDataVertex MakeInputDataVertex(const FName InName, const FString& InFriendlyName)
+		const FLazyName InputBaseName{"In"};
+
+		FName MakeInputVertexName(const int32 InIndex)
 		{
-			const FText DisplayName = METASOUND_LOCTEXT_FORMAT(
-				"DisplayName_{0}", "{0}", FText::FromName(InName));
+			FName Name = InputBaseName;
+			Name.SetNumber(InIndex);
+			return Name;
+		}
 
-			const FText Tooltip = METASOUND_LOCTEXT_FORMAT(
-				"Tooltip_{0}", "Channel {0}", FText::FromString(InFriendlyName));
-			
-			return FInputDataVertex
-			{
-				InName,
-				GetMetasoundDataTypeName<FAudioBuffer>(),
-				FDataVertexMetadata
-				{
-					.Description = Tooltip,
-					.DisplayName = DisplayName,
-					.bIsAdvancedDisplay = false
-				}
-			};
+		FInputDataVertex MakeInputDataVertex(const int32 InIndex, const FName InName, const FString& InFriendlyName)
+		{
+			const FName InputName = MakeInputVertexName(InIndex);
+
+			const FText InputDisplayName = METASOUND_LOCTEXT_FORMAT("In_DisplayName", "{0}", FText::FromString(InName.ToString()));
+
+			return TInputDataVertex<FAudioBuffer>{InputName, FDataVertexMetadata{FText::FromString(*InFriendlyName), InputDisplayName}};
 		}
 		
-		static FVertexInterface MakeClassInterface(const FName Format)
+		
+		FVertexInterface MakeClassInterface(const FName Format)
 		{
 			const Audio::FChannelTypeFamily* Found = Audio::GetChannelRegistry().FindConcreteChannel(Format);
 			if (!Found)
@@ -57,25 +55,15 @@ namespace Metasound
 			{
 				if (TOptional<Audio::FChannelTypeFamily::FChannelName> Name = Found->GetChannelName(ChannelIndex); Name.IsSet())
 				{
-					InputInterface.Add(MakeInputDataVertex(Name->Name, Name->FriendlyName));
+					InputInterface.Add(MakeInputDataVertex(ChannelIndex, Name->Name, Name->FriendlyName));
 				}
 				else
 				{
-					const FName Input("In", ChannelIndex);
-					InputInterface.Add(MakeInputDataVertex(Input, FString::FromInt(ChannelIndex) ));	
+					InputInterface.Add(MakeInputDataVertex(ChannelIndex, TEXT("Input"), FString::FromInt(ChannelIndex) ));	
 				}
 			}
-			
-			// outputs
 			FOutputVertexInterface OutputInterface;
-			OutputInterface.Add(
-				FOutputDataVertex(
-					METASOUND_GET_PARAM_NAME(OutputToCat),
-					Format,
-					METASOUND_GET_PARAM_METADATA(OutputToCat),
-					EVertexAccessType::Reference
-				)
-			);
+            OutputInterface.Add(TOutputDataVertex<FChannelAgnosticType>(METASOUND_GET_PARAM_NAME_AND_METADATA(OutputToCat)));
 			return{
 				MoveTemp(InputInterface),
 				MoveTemp(OutputInterface)
@@ -115,26 +103,22 @@ namespace Metasound
 		}
 		virtual ~FCatMakeOperator() override = default;
 
-		static TDataWriteReference<FChannelAgnosticType>
-		MakePolyCat(const FName InDerivedType, const FOperatorSettings& InOpSettings)
-		{
-			using namespace Frontend;
-			const FLiteral Literal(InDerivedType.ToString());
-			TOptional<FAnyDataReference> Ref = IDataTypeRegistry::Get().CreateDataReference(InDerivedType, EDataReferenceAccessType::Write, Literal, InOpSettings);
-			return Ref->GetAs<TDataWriteReference<FChannelAgnosticType>>();
-		}
-
-		static FName GetPinName(const int32 InIndex, const Audio::FChannelTypeFamily& InFormat)
+		static const FVertexInterface& GetDefaultInterface()
 		{
-			if (auto Names = InFormat.GetChannelName(InIndex))
+			using namespace CatMakeNodePrivate ;
+			auto CreateDefaultInterface = []()-> FVertexInterface
 			{
-				return Names->Name;
-			}
-			FName Input;
-			Input.SetNumber(InIndex);
-			return Input;
+				// outputs
+				FOutputVertexInterface OutputInterface;
+				OutputInterface.Add(TOutputDataVertex<FChannelAgnosticType>(METASOUND_GET_PARAM_NAME_AND_METADATA(OutputToCat)));
+				
+				return FVertexInterface({}, OutputInterface);
+			}; // end lambda: CreateDefaultInterface()
+
+			static const FVertexInterface DefaultInterface = CreateDefaultInterface();
+			return DefaultInterface;
 		}
-		
+			
 		static TUniquePtr<IOperator> CreateOperator(const FBuildOperatorParams& InParams, FBuildResults& OutResults)
 		{
 			using namespace CatMakeNodePrivate;
@@ -147,20 +131,16 @@ namespace Metasound
 			const FName OutputFormat = CatMakeData->GetToType();
 			const FInputVertexInterfaceData& InputData = InParams.InputData;
 			const Audio::FChannelTypeFamily* ConcreteToType = Audio::GetChannelRegistry().FindConcreteChannel(OutputFormat);
-			if (!ConcreteToType)
-			{
-				return TUniquePtr<FNoOpOperator>();
-			}
 
 			// Set the number of inputs pins to match the output format.
 			TArray<TDataReadReference<FAudioBuffer>> InputAudioVertices;
 			for (int32 i = 0; i < ConcreteToType->NumChannels(); ++i)
 			{
-				InputAudioVertices.Emplace(InputData.GetOrCreateDefaultDataReadReference<FAudioBuffer>(GetPinName(i, *ConcreteToType), InParams.OperatorSettings));
+				InputAudioVertices.Emplace(InputData.GetOrCreateDefaultDataReadReference<FAudioBuffer>(MakeInputVertexName(i), InParams.OperatorSettings));
 			}
 						
 			// Make output CAT to match our configs settings. (use concrete form).
-			FChannelAgnosticTypeWriteRef OutputPin = MakePolyCat(ConcreteToType->GetName(), InParams.OperatorSettings);
+			FChannelAgnosticTypeWriteRef OutputPin = FChannelAgnosticTypeWriteRef::CreateNew(InParams.OperatorSettings, ConcreteToType->GetName());
 			
 			return MakeUnique<FCatMakeOperator>(
 				InParams,
@@ -174,7 +154,7 @@ namespace Metasound
 			using namespace CatMakeNodePrivate;
 			for (int32 i = 0; i < InputAudioVertices.Num(); i++)
 			{
-				InOutVertexData.BindReadVertex(GetPinName(i, OutputCat->GetType()), InputAudioVertices[i]);
+				InOutVertexData.BindReadVertex(MakeInputVertexName(i), InputAudioVertices[i]);
 			}
 		}
 		virtual void BindOutputs(FOutputVertexInterfaceData& InOutVertexData) override
@@ -197,14 +177,11 @@ namespace Metasound
 				TArrayView<float> Dst = OutputCat->GetChannel(Index);
 				checkSlow(Src.Num() == Dst.Num());
 				FMemory::Memcpy(Dst.GetData(), Src.GetData(), Dst.Num() * sizeof(float));
-
-				// REWRITE ME TO BE REFERENCE FORWARD.
 			}
 		}
 		static FNodeClassMetadata GetNodeInfo()
 		{
 			using namespace CatMakeNodePrivate;
-			static FName DefaultFormat = TEXT("Cat:Mono1Dot0");
 			return FNodeClassMetadata
 			{
 				FNodeClassName{ "Experimental", "CatMakeOperator", "" },
@@ -213,8 +190,8 @@ namespace Metasound
 				METASOUND_LOCTEXT("CatMakeNodeName", "CAT Make Node"),	
 				METASOUND_LOCTEXT("CatMakeNodeNameDescription", "A Node that builds CATs"),	
 				TEXT("UE"), // Author
-				METASOUND_LOCTEXT("CatMakeNodePromptIfMissing", "Enable the MetaSoundExperimental Plugin"), // Prompt if missing
-				MakeClassInterface(DefaultFormat),
+				METASOUND_LOCTEXT("ExampleConfigurablePromptIfMissing", "Enable the MetaSoundExperimental Plugin"), // Prompt if missing
+				GetDefaultInterface(),
 				{}
 			};
 		}
diff --git a/Engine/Plugins/Experimental/MetasoundExperimental/Source/MetasoundExperimentalRuntime/Private/MetasoundCatMakeNode.h b/Engine/Plugins/Experimental/MetasoundExperimental/Source/MetasoundExperimentalRuntime/Private/MetasoundCatMakeNode.h
--- a/Engine/Plugins/Experimental/MetasoundExperimental/Source/MetasoundExperimentalRuntime/Private/MetasoundCatMakeNode.h
+++ b/Engine/Plugins/Experimental/MetasoundExperimental/Source/MetasoundExperimentalRuntime/Private/MetasoundCatMakeNode.h
@@ -14,8 +14,8 @@ struct FMetaSoundCatMakeNodeConfiguration : public FMetaSoundFrontendNodeConfigu
 
 	FMetaSoundCatMakeNodeConfiguration() = default;
 
-	UPROPERTY(EditAnywhere, Category = General, meta = (GetOptions="MetasoundCatCastingOptionsHelper.GetCastingOptions_NoAbstract"))
-	FName Format = TEXT("Cat:Stereo");
+	UPROPERTY(EditAnywhere, Category = General, meta = (GetOptions="MetasoundCatCastingOptionsHelper.GetCastingOptions"))
+	FName Format = TEXT("Stereo");
 
 	/* Get the current interface for the class based upon the node extension */
 	virtual TInstancedStruct<FMetasoundFrontendClassInterface> OverrideDefaultInterface(const FMetasoundFrontendClass& InNodeClass) const override;
diff --git a/Engine/Plugins/Experimental/MetasoundExperimental/Source/MetasoundExperimentalRuntime/Private/MetasoundCatMixerNode.cpp b/Engine/Plugins/Experimental/MetasoundExperimental/Source/MetasoundExperimentalRuntime/Private/MetasoundCatMixerNode.cpp
--- a/Engine/Plugins/Experimental/MetasoundExperimental/Source/MetasoundExperimentalRuntime/Private/MetasoundCatMixerNode.cpp
+++ b/Engine/Plugins/Experimental/MetasoundExperimental/Source/MetasoundExperimentalRuntime/Private/MetasoundCatMixerNode.cpp
@@ -2,10 +2,6 @@
 
 #include "MetasoundCatMixerNode.h"
 
-#include "Algo/MaxElement.h"
-#include "DSP/FloatArrayMath.h"
-#include "DSP/MultiMono.h"
-#include "MetasoundCatBreakNode.h"
 #include "MetasoundChannelAgnosticType.h"
 #include "MetasoundExecutableOperator.h"
 #include "MetasoundFacade.h"
@@ -13,7 +9,11 @@
 #include "MetasoundPrimitives.h"
 #include "MetasoundStandardNodesCategories.h"
 #include "MetasoundVertex.h"
+#include "Algo/MaxElement.h"
+#include "DSP/FloatArrayMath.h"
 #include "TypeFamily/ChannelTypeFamily.h"
+#include "MetasoundCatBreakNode.h"
+#include "ChannelAgnostic/ChannelAgnosticTypeUtils.h"
 
 #define LOCTEXT_NAMESPACE "MetasoundStandardNodes_CatMixerNode"
 
@@ -71,7 +71,7 @@ namespace Metasound
 			return TInputDataVertex<float>{InputName, FDataVertexMetadata{ FText::FromString(*InFriendlyName), InputDisplayName }, 1.0f};
 		}
 
-		static FVertexInterface GetVertexInterface(const int32 InNumInputs, const FName InOutputFormat)
+		FVertexInterface GetVertexInterface(const int32 InNumInputs)
 		{
 			FInputVertexInterface InputInterface;
 			for (int32 i = 0; i < InNumInputs; i++)
@@ -81,14 +81,7 @@ namespace Metasound
 			}
 
 			FOutputVertexInterface OutputInterface;
-			OutputInterface.Add(
-				FOutputDataVertex(
-					METASOUND_GET_PARAM_NAME(OutputCat),
-					InOutputFormat,
-					METASOUND_GET_PARAM_METADATA(OutputCat),
-					EVertexAccessType::Reference)
-				);
-			
+			OutputInterface.Add(TOutputDataVertex<FChannelAgnosticType>(METASOUND_GET_PARAM_NAME_AND_METADATA(OutputCat)));
 			return FVertexInterface
 				{
 					MoveTemp(InputInterface),
@@ -167,13 +160,12 @@ namespace Metasound
 
 		static const FNodeClassMetadata& GetNodeInfo()
 		{
-			static const FName DefaultOutputFormat = TEXT("Cat");
 			static const FNodeClassMetadata Metadata = CreateNodeClassMetadata
 			(
 				FName(TEXT("CAT Mixer Node")),
 				METASOUND_LOCTEXT("Metasound_CatMixerNodeDisplayName", "CAT Mixer Node"),
 				METASOUND_LOCTEXT("Metasound_CatMixerNodeDescription", "CAT Mixer Node"),
-				CatMixerPrivate::GetVertexInterface(1, DefaultOutputFormat)
+				CatMixerPrivate::GetVertexInterface(1)
 			);
 			return Metadata;
 		}
@@ -328,10 +320,7 @@ TInstancedStruct<FMetasoundFrontendClassInterface> FMetaSoundCatMixingNodeConfig
 {
 	using namespace Metasound;
 	using namespace CatMixerPrivate;
-	return TInstancedStruct<FMetasoundFrontendClassInterface>::Make(
-		FMetasoundFrontendClassInterface::GenerateClassInterface(
-				GetVertexInterface(NumInputs, CustomMixFormat )
-				));
+	return TInstancedStruct<FMetasoundFrontendClassInterface>::Make(FMetasoundFrontendClassInterface::GenerateClassInterface(GetVertexInterface(NumInputs)));
 }
 
 TSharedPtr<const Metasound::IOperatorData> FMetaSoundCatMixingNodeConfiguration::GetOperatorData() const
diff --git a/Engine/Plugins/Experimental/MetasoundExperimental/Source/MetasoundExperimentalRuntime/Private/MetasoundCatMixerNode.h b/Engine/Plugins/Experimental/MetasoundExperimental/Source/MetasoundExperimentalRuntime/Private/MetasoundCatMixerNode.h
--- a/Engine/Plugins/Experimental/MetasoundExperimental/Source/MetasoundExperimentalRuntime/Private/MetasoundCatMixerNode.h
+++ b/Engine/Plugins/Experimental/MetasoundExperimental/Source/MetasoundExperimentalRuntime/Private/MetasoundCatMixerNode.h
@@ -35,7 +35,7 @@ struct FMetaSoundCatMixingNodeConfiguration : public FMetaSoundFrontendNodeConfi
 	EMetasoundCatCastingMethod CatCastingMethod = EMetasoundCatCastingMethod::ChannelDrop;
 
 	UPROPERTY(EditAnywhere, Category = General, meta = (GetOptions="MetasoundCatCastingOptionsHelper.GetCastingOptions", EditCondition = "FormatChoosingMethod == EMetasoundMixerFormatChoosingMethod::Custom", EditConditionHides))
-	FName CustomMixFormat = TEXT("Cat");
+	FName CustomMixFormat = TEXT("Mono");
 
 	UPROPERTY(EditAnywhere, Category = General, meta = (ClampMin=1, ClampMax=100))
 	int32 NumInputs = 1;
diff --git a/Engine/Plugins/Experimental/MetasoundExperimental/Source/MetasoundExperimentalRuntime/Private/MetasoundCatPannerNode.cpp b/Engine/Plugins/Experimental/MetasoundExperimental/Source/MetasoundExperimentalRuntime/Private/MetasoundCatPannerNode.cpp
deleted file mode 100644
diff --git a/Engine/Plugins/Experimental/MetasoundExperimental/Source/MetasoundExperimentalRuntime/Private/MetasoundChannelAgnosticType.cpp b/Engine/Plugins/Experimental/MetasoundExperimental/Source/MetasoundExperimentalRuntime/Private/MetasoundChannelAgnosticType.cpp
--- a/Engine/Plugins/Experimental/MetasoundExperimental/Source/MetasoundExperimentalRuntime/Private/MetasoundChannelAgnosticType.cpp
+++ b/Engine/Plugins/Experimental/MetasoundExperimental/Source/MetasoundExperimentalRuntime/Private/MetasoundChannelAgnosticType.cpp
@@ -3,7 +3,6 @@
 #include "MetasoundChannelAgnosticType.h"
 #include "TypeFamily/TypeFamily.h"
 #include "TypeFamily/ChannelTypeFamily.h"
-#include "MetasoundPolymorphic.h"
 
 namespace Metasound
 {
@@ -20,7 +19,7 @@ namespace Metasound
 	 */
 	FName FChannelAgnosticType::GetDefaultCatFormat()
 	{
-		static const FName DefaultFormat(TEXT("Cat:Mono1Dot0"));
+		static const FName DefaultFormat(TEXT("Mono")); // 1 channel, Mono buffer.
 		check(Audio::GetChannelRegistry().FindConcreteChannel(DefaultFormat) != nullptr);
 		return DefaultFormat;
 	}
@@ -30,25 +29,9 @@ namespace Metasound
 	{}
 	
 	FChannelAgnosticType::FChannelAgnosticType(const int32 InNumFramesPerBlock, const FName& InChannelTypeName)
-		: Super(Audio::GetChannelRegistry().FindConcreteChannelChecked(InChannelTypeName), InNumFramesPerBlock)
+		: Super(*Audio::GetChannelRegistry().FindConcreteChannel(InChannelTypeName), InNumFramesPerBlock)
 	{}
-	
-	// Register.
-	REGISTER_METASOUND_DATATYPE(FChannelAgnosticType, "Cat");
-	REGISTER_METASOUND_POLY_TYPE(FChannelAgnosticType);
-
-	REGISTER_METASOUND_DATATYPE(FCompositeChannelAgnosticType, "Cat:Composite");
-	REGISTER_METASOUND_POLY_TYPE(FCompositeChannelAgnosticType);
-
-	REGISTER_METASOUND_DATATYPE(FDiscreteChannelAgnosticType, "Cat:Discrete");
-	REGISTER_METASOUND_POLY_TYPE(FDiscreteChannelAgnosticType);
-
-	REGISTER_METASOUND_DATATYPE(FSoundfieldChannelAgnosticType, "Cat:Soundfield");
-	REGISTER_METASOUND_POLY_TYPE(FSoundfieldChannelAgnosticType);
 }
 
-
-
-
-
-
+// Register.
+REGISTER_METASOUND_DATATYPE(Metasound::FChannelAgnosticType, "Cat");
\ No newline at end of file
diff --git a/Engine/Plugins/Experimental/MetasoundExperimental/Source/MetasoundExperimentalRuntime/Private/MetasoundExperimentalRuntimeModule.cpp b/Engine/Plugins/Experimental/MetasoundExperimental/Source/MetasoundExperimentalRuntime/Private/MetasoundExperimentalRuntimeModule.cpp
--- a/Engine/Plugins/Experimental/MetasoundExperimental/Source/MetasoundExperimentalRuntime/Private/MetasoundExperimentalRuntimeModule.cpp
+++ b/Engine/Plugins/Experimental/MetasoundExperimental/Source/MetasoundExperimentalRuntime/Private/MetasoundExperimentalRuntimeModule.cpp
@@ -2,17 +2,12 @@
 
 #include "MetasoundExperimentalRuntimeModule.h"
 
-#include "MetasoundCatConversionNode.h"
-#include "MetasoundCatDataSchemas.h"
 #include "MetasoundFrontendModuleRegistrationMacros.h"
 #include "Modules/ModuleManager.h"
 
 void FMetasoundExperimentalRuntimeModule::StartupModule()
 {
 	FModuleManager::Get().LoadModule("AudioExperimentalRuntime");
-	Metasound::LoadPolymorphicDataSchemas();
-	Metasound::RegisterCatConvertNode();
-	
 	METASOUND_REGISTER_ITEMS_IN_MODULE
 }
 
diff --git a/Engine/Plugins/Experimental/MetasoundExperimental/Source/MetasoundExperimentalRuntime/Public/MetasoundCatPannerNode.h b/Engine/Plugins/Experimental/MetasoundExperimental/Source/MetasoundExperimentalRuntime/Public/MetasoundCatPannerNode.h
deleted file mode 100644
diff --git a/Engine/Plugins/Experimental/MetasoundExperimental/Source/MetasoundExperimentalRuntime/Public/MetasoundChannelAgnosticType.h b/Engine/Plugins/Experimental/MetasoundExperimental/Source/MetasoundExperimentalRuntime/Public/MetasoundChannelAgnosticType.h
--- a/Engine/Plugins/Experimental/MetasoundExperimental/Source/MetasoundExperimentalRuntime/Public/MetasoundChannelAgnosticType.h
+++ b/Engine/Plugins/Experimental/MetasoundExperimental/Source/MetasoundExperimentalRuntime/Public/MetasoundChannelAgnosticType.h
@@ -2,11 +2,9 @@
 
 #pragma once
 
-#include "MetasoundPolymorphic.h"
 #include "ChannelAgnostic/ChannelAgnosticType.h"
 #include "MetasoundDataTypeRegistrationMacro.h"
-
-#include "MetasoundChannelAgnosticType.generated.h"
+#include "DSP/MultiMono.h"
 
 namespace Metasound
 {
@@ -22,119 +20,22 @@ namespace Metasound
 		 */
 		static FName GetDefaultCatFormat();
 
-		METASOUNDEXPERIMENTALRUNTIME_API explicit FChannelAgnosticType(const FOperatorSettings& InSettings, const FName& InChannelTypeName = GetDefaultCatFormat());
+		explicit FChannelAgnosticType(const FOperatorSettings& InSettings, const FName& InChannelTypeName = GetDefaultCatFormat());
 		
 		explicit FChannelAgnosticType(const int32 InNumFramesPerBlock, const FName& InChannelTypeName = GetDefaultCatFormat());
-
-		static FName LiteralToFName(const FString& In)
-		{
-			if (In.IsEmpty())
-			{
-				static const FName Default = TEXT("Cat:Mono1Dot0");
-				return Default;
-			}
-			return FName(In);
-		}
+				
 		/* FLiteral compatible Constructor
 		 * TODO: Literals only support FString and not FNames, so there is a conversion here.
 		 */
-		METASOUNDEXPERIMENTALRUNTIME_API explicit FChannelAgnosticType(const FOperatorSettings& InSettings, const FString& InChannelTypeName)
-			: FChannelAgnosticType(InSettings, LiteralToFName((InChannelTypeName)))
+		explicit FChannelAgnosticType(const FOperatorSettings& InSettings, const FString& InChannelTypeName)
+			: FChannelAgnosticType(InSettings, FName(*InChannelTypeName))
 		{}
 	};
+	
 
-	class FDiscreteChannelAgnosticType : public FChannelAgnosticType
-	{
-		using Super = FChannelAgnosticType;
-		using Super::Super;
-	};
-	class FSoundfieldChannelAgnosticType : public FChannelAgnosticType
-	{
-		using Super = FChannelAgnosticType;
-		using Super::Super;
-	};
-	class FCompositeChannelAgnosticType : public FChannelAgnosticType
-	{
-		using Super = FChannelAgnosticType;
-		using Super::Super;
-	};
+	
 }
 
-USTRUCT()
-struct FCatDataTypeSchema
-{
-	GENERATED_BODY()
-
-	UPROPERTY()
-	FName Name;
-
-	UPROPERTY()
-	FName BaseType;
-
-	UPROPERTY()
-	FName FamilyType;
-
-	UPROPERTY()
-	FString FriendlyName;
-
-	UPROPERTY()
-	bool bIsAbstract = false;
-
-	UPROPERTY()
-	bool bIsParentsDefault = false;
-};
-
-USTRUCT()
-struct FCatSoundfieldDataSchema : public FCatDataTypeSchema
-{
-	GENERATED_BODY()
-
-	UPROPERTY()
-	int32 NumOrders = 1;
-};
-
-USTRUCT()
-struct FCatDiscreteChannel
-{
-	GENERATED_BODY()
-
-	UPROPERTY()
-	float ElevationDegrees = 0;			// In Degrees. +180 to -180
-
-	UPROPERTY()
-	float AzimuthDegrees = 0;			// -180 to 180 azimuth.
-
-	UPROPERTY()
-	FName ID;							// This should really be the speaker enum.
-
-	UPROPERTY()
-	bool bIsSpatialized = true;			// If this is off, azimuth+degrees do nothing. (for excluding from panner). 
-};
-
-USTRUCT()
-struct FCatDiscreteDataSchema : public FCatDataTypeSchema
-{
-	GENERATED_BODY()
-
-	UPROPERTY()
-	TArray<FCatDiscreteChannel> Order;
-};
-
-USTRUCT()
-struct FCatCompositeDataSchema : public FCatDataTypeSchema
-{
-	GENERATED_BODY()
-};
-
-namespace  Metasound
-{
-	DECLARE_METASOUND_DATA_REFERENCE_TYPES(Metasound::FChannelAgnosticType,				METASOUNDEXPERIMENTALRUNTIME_API, FChannelAgnosticTypeTypeInfo,				FChannelAgnosticTypeReadRef,			FChannelAgnosticTypeWriteRef);
-	DECLARE_METASOUND_DATA_REFERENCE_TYPES(Metasound::FCompositeChannelAgnosticType,	METASOUNDEXPERIMENTALRUNTIME_API, FCompositeChannelAgnosticTypeTypeInfo,	FCompositeChannelAgnosticTypeReadRef,	FCompositeChannelAgnosticTypeWriteRef);
-	DECLARE_METASOUND_DATA_REFERENCE_TYPES(Metasound::FDiscreteChannelAgnosticType,		METASOUNDEXPERIMENTALRUNTIME_API, FDiscreteChannelAgnosticTypeTypeInfo,		FDiscreteChannelAgnosticTypeReadRef,	FDiscreteChannelAgnosticTypeWriteRef);
-	DECLARE_METASOUND_DATA_REFERENCE_TYPES(Metasound::FSoundfieldChannelAgnosticType,	METASOUNDEXPERIMENTALRUNTIME_API, FSoundfieldChannelAgnosticTypeTypeInfo,	FSoundfieldChannelAgnosticTypeReadRef,	FSoundChannelAgnosticTypeWriteRef);
-
-	DECLARE_METASOUND_POLY_TYPE(FChannelAgnosticType,			void, true);
-	DECLARE_METASOUND_POLY_TYPE(FCompositeChannelAgnosticType,	Metasound::FChannelAgnosticType, true);
-	DECLARE_METASOUND_POLY_TYPE(FDiscreteChannelAgnosticType,	Metasound::FChannelAgnosticType, true);
-	DECLARE_METASOUND_POLY_TYPE(FSoundfieldChannelAgnosticType,	Metasound::FChannelAgnosticType, true);
-}
\ No newline at end of file
+// Declare it.
+DECLARE_METASOUND_DATA_REFERENCE_TYPES(Metasound::FChannelAgnosticType, METASOUNDEXPERIMENTALRUNTIME_API,
+	FChannelAgnosticTypeTypeInfo, FChannelAgnosticTypeReadRef, FChannelAgnosticTypeWriteRef);
diff --git a/Engine/Plugins/Runtime/Metasound/Source/MetasoundFrontend/Private/MetasoundFrontendDataTypeRegistry.cpp b/Engine/Plugins/Runtime/Metasound/Source/MetasoundFrontend/Private/MetasoundFrontendDataTypeRegistry.cpp
--- a/Engine/Plugins/Runtime/Metasound/Source/MetasoundFrontend/Private/MetasoundFrontendDataTypeRegistry.cpp
+++ b/Engine/Plugins/Runtime/Metasound/Source/MetasoundFrontend/Private/MetasoundFrontendDataTypeRegistry.cpp
@@ -719,19 +719,15 @@ PRAGMA_ENABLE_DEPRECATION_WARNINGS
 					{
 						if (Entry->GetDataTypeInfo().bIsParsable)
 						{
-							const bool bIsAbstract = Entry->GetDataTypeInfo().bIsAbstract;;
+							NodeRegistry->RegisterNode(MakeUnique<FInputNodeRegistryEntry>(Entry));
 							
-							NodeRegistry->RegisterNode(MakeUnique<FInputNodeRegistryEntry>(Entry));							
 							NodeRegistry->RegisterNode(MakeUnique<FOutputNodeRegistryEntry>(Entry));
-
-							if (!bIsAbstract)
-							{
-								NodeRegistry->RegisterNode(MakeUnique<FLiteralNodeRegistryEntry>(Entry));
-								NodeRegistry->RegisterNode(MakeUnique<FVariableNodeRegistryEntry>(Entry));
-								NodeRegistry->RegisterNode(MakeUnique<FVariableMutatorNodeRegistryEntry>(Entry));
-								NodeRegistry->RegisterNode(MakeUnique<FVariableAccessorNodeRegistryEntry>(Entry));
-								NodeRegistry->RegisterNode(MakeUnique<FVariableDeferredAccessorNodeRegistryEntry>(Entry));
-							}
+							
+							NodeRegistry->RegisterNode(MakeUnique<FLiteralNodeRegistryEntry>(Entry));
+							NodeRegistry->RegisterNode(MakeUnique<FVariableNodeRegistryEntry>(Entry));
+							NodeRegistry->RegisterNode(MakeUnique<FVariableMutatorNodeRegistryEntry>(Entry));
+							NodeRegistry->RegisterNode(MakeUnique<FVariableAccessorNodeRegistryEntry>(Entry));
+							NodeRegistry->RegisterNode(MakeUnique<FVariableDeferredAccessorNodeRegistryEntry>(Entry));
 
 							if (Entry->GetDataTypeInfo().bIsConstructorType)
 							{
diff --git a/Engine/Plugins/Runtime/Metasound/Source/MetasoundFrontend/Private/MetasoundOutputNode.cpp b/Engine/Plugins/Runtime/Metasound/Source/MetasoundFrontend/Private/MetasoundOutputNode.cpp
--- a/Engine/Plugins/Runtime/Metasound/Source/MetasoundFrontend/Private/MetasoundOutputNode.cpp
+++ b/Engine/Plugins/Runtime/Metasound/Source/MetasoundFrontend/Private/MetasoundOutputNode.cpp
@@ -67,7 +67,7 @@ namespace Metasound
 			// Use data reference if it is passed in. 
 			if (const FAnyDataReference* Ref = InParams.InputData.FindDataReference(VertexName))
 			{
-				checkf(IsCastable(Ref->GetDataTypeName(), Vertex.DataTypeName), TEXT("Mismatched data type names on output node (%s). Expected (%s), received (%s)."), *VertexName.ToString(), *Vertex.DataTypeName.ToString(), *(Ref->GetDataTypeName().ToString()));
+				checkf(Ref->GetDataTypeName() == Vertex.DataTypeName, TEXT("Mismatched data type names on output node (%s). Expected (%s), received (%s)."), *VertexName.ToString(), *Vertex.DataTypeName.ToString(), *(Ref->GetDataTypeName().ToString()));
 				return MakeUnique<FOutputOperator>(VertexName, *Ref);
 			}
 			else 
diff --git a/Engine/Plugins/Runtime/Metasound/Source/MetasoundFrontend/Public/MetasoundNodeConstructorParams.h b/Engine/Plugins/Runtime/Metasound/Source/MetasoundFrontend/Public/MetasoundNodeConstructorParams.h
--- a/Engine/Plugins/Runtime/Metasound/Source/MetasoundFrontend/Public/MetasoundNodeConstructorParams.h
+++ b/Engine/Plugins/Runtime/Metasound/Source/MetasoundFrontend/Public/MetasoundNodeConstructorParams.h
@@ -40,9 +40,6 @@ namespace Metasound
 		FVertexName VertexName;
 		FLiteral InitParam = FLiteral::CreateInvalid();
 
-		// If defined this is the polymorphic type name to use.
-		FName PolymorphicTypeName;
-
 		UE_DEPRECATED(5.3, "Transmission is no longer enabled on input nodes. In order set values on a MetaSound input, access them through the bound FInputVertexInterfaceData.")
 		bool bEnableTransmission = false;
 	};
diff --git a/Engine/Plugins/Runtime/Metasound/Source/MetasoundGraphCore/Private/MetasoundOperatorBuilder.cpp b/Engine/Plugins/Runtime/Metasound/Source/MetasoundGraphCore/Private/MetasoundOperatorBuilder.cpp
--- a/Engine/Plugins/Runtime/Metasound/Source/MetasoundGraphCore/Private/MetasoundOperatorBuilder.cpp
+++ b/Engine/Plugins/Runtime/Metasound/Source/MetasoundGraphCore/Private/MetasoundOperatorBuilder.cpp
@@ -680,7 +680,7 @@ PRAGMA_ENABLE_DEPRECATION_WARNINGS
 
 				if (const FAnyDataReference* DataReference = NodeInputData.FindDataReference(InputDestination.Vertex.VertexName))
 				{
-					if (IsCastable(DataReference->GetDataTypeName(), InputDestination.Vertex.DataTypeName))
+					if (DataReference->GetDataTypeName() == InputDestination.Vertex.DataTypeName)
 					{
 						bFoundDataReference = true;
 						OutVertexData.GetInputs().SetVertex(InputDestination.Vertex.VertexName, *DataReference);
@@ -708,7 +708,7 @@ PRAGMA_ENABLE_DEPRECATION_WARNINGS
 				const FOutputVertexInterfaceData& NodeOutputData = OperatorInfo->VertexData.GetOutputs();
 				if (const FAnyDataReference* DataReference = NodeOutputData.FindDataReference(OutputSource.Vertex.VertexName))
 				{
-					if (IsCastable(DataReference->GetDataTypeName() , OutputSource.Vertex.DataTypeName))
+					if (DataReference->GetDataTypeName() == OutputSource.Vertex.DataTypeName)
 					{
 						bFoundDataReference = true;
 						OutVertexData.GetOutputs().SetVertex(OutputSource.Vertex.VertexName, FAnyDataReference{*DataReference});
diff --git a/Engine/Plugins/Runtime/Metasound/Source/MetasoundGraphCore/Private/MetasoundVertexData.cpp b/Engine/Plugins/Runtime/Metasound/Source/MetasoundGraphCore/Private/MetasoundVertexData.cpp
--- a/Engine/Plugins/Runtime/Metasound/Source/MetasoundGraphCore/Private/MetasoundVertexData.cpp
+++ b/Engine/Plugins/Runtime/Metasound/Source/MetasoundGraphCore/Private/MetasoundVertexData.cpp
@@ -81,7 +81,7 @@ namespace Metasound
 
 		void FInputBinding::Set(FAnyDataReference&& InAnyDataReference)
 		{
-			check(IsCastable(InAnyDataReference.GetDataTypeName(),Vertex.DataTypeName));
+			check(IsCastable(Vertex.DataTypeName,InAnyDataReference.GetDataTypeName()));
 			Data.Emplace(MoveTemp(InAnyDataReference));
 			CheckAccessTypeCompatibility(Vertex, *Data);
 		}
@@ -128,7 +128,7 @@ namespace Metasound
 
 		void FInputBinding::Bind(FInputBinding& InBinding)
 		{
-			check(IsCastable(InBinding.GetVertex().DataTypeName,Vertex.DataTypeName));
+			check(IsCastable(Vertex.DataTypeName, InBinding.GetVertex().DataTypeName));
 
 			if (Data.IsSet()) 
 			{ 
@@ -166,7 +166,7 @@ namespace Metasound
 
 		void FOutputBinding::Set(FAnyDataReference&& InAnyDataReference)
 		{
-			check(IsCastable(InAnyDataReference.GetDataTypeName(),Vertex.DataTypeName));
+			check(IsCastable(Vertex.DataTypeName, InAnyDataReference.GetDataTypeName()));
 			Data.Emplace(MoveTemp(InAnyDataReference));
 			CheckAccessTypeCompatibility(Vertex, *Data);
 		}
@@ -251,10 +251,7 @@ namespace Metasound
 		{
 			if (BindingType* Binding = Find(InBindings, InVertexName))
 			{
-				if (IsCastable(
-					InDataReference.GetDataTypeName(),						// Incoming ref. (dynamic reference).
-					Binding->GetVertex().DataTypeName						// Type of this pin. (must be derived from this or equal).
-					))
+				if (IsCastable(Binding->GetVertex().DataTypeName,InDataReference.GetDataTypeName()))
 				{
 					Binding->Set(MoveTemp(InDataReference));
 				}
diff --git a/Engine/Plugins/Runtime/Metasound/Source/MetasoundGraphCore/Public/MetasoundDataReference.h b/Engine/Plugins/Runtime/Metasound/Source/MetasoundGraphCore/Public/MetasoundDataReference.h
--- a/Engine/Plugins/Runtime/Metasound/Source/MetasoundGraphCore/Public/MetasoundDataReference.h
+++ b/Engine/Plugins/Runtime/Metasound/Source/MetasoundGraphCore/Public/MetasoundDataReference.h
@@ -245,30 +245,6 @@ namespace Metasound
 			);
 		}
 
-		template <typename Derived, typename... ArgTypes>
-		UE_EXPERIMENTAL(5.8, "Polymorphic API is under development")
-		static TDataReference CreateNewDerivedAs(const FName& InDerivedTypeName,
-												 FMetasoundDataTypeId InDerivedTypeId,
-												 ArgTypes&&... Args)
-		{
-			static_assert(std::is_base_of_v<DataType, Derived>, "Derived must inherit from DataType");
-			static_assert(std::is_constructible_v<Derived, ArgTypes...>, "Args don't match a Derived constructor");
-
-		#if DO_CHECK
-			PRAGMA_DISABLE_EXPERIMENTAL_WARNINGS
-			check(Metasound::IsChildOfByName(InDerivedTypeName, GetMetasoundDataTypeName<DataType>()));
-			PRAGMA_ENABLE_EXPERIMENTAL_WARNINGS
-		#endif
-
-			return TDataReference(
-				EDataRefShouldConstruct::NewObject,
-				MetasoundDataReferencePrivate::MakeRefOf<DataType, Derived>(Forward<ArgTypes>(Args)...),
-				InDerivedTypeName,
-				InDerivedTypeId
-			);
-		}
-
-
 		/** Enable copy constructor */
 		TDataReference(const TDataReference<DataType>& Other) = default;
 
@@ -361,12 +337,6 @@ namespace Metasound
 
 		template<typename T>
 		friend TDataValueReference<T> ValueCast(const TDataWriteReference<T>& InRef);
-		
-		using FRefType = TDataReference<DataType>::FRefType;
-		TDataValueReference(EDataRefShouldConstruct, FRefType InRef, FName InTypeName, const void* InTypeId)
-			: TDataReference<DataType>(EDataRefShouldConstruct::NewObject, MoveTemp(InRef), InTypeName, InTypeId)
-		{
-		}
 
 	public:
 
@@ -390,29 +360,6 @@ namespace Metasound
 				MetasoundDataReferencePrivate::MakeRefOf<DataType, Derived>(Forward<ArgTypes>(Args)...)
 			);
 		}
-
-		template <typename Derived, typename... ArgTypes>
-		UE_EXPERIMENTAL(5.8, "Polymorphic API is under development")
-		static TDataValueReference CreateNewDerivedAs(const FName& InDerivedTypeName,
-													 FMetasoundDataTypeId InDerivedTypeId,
-													 ArgTypes&&... Args)
-		{
-			static_assert(std::is_base_of_v<DataType, Derived>, "Derived must inherit from DataType");
-			static_assert(std::is_constructible_v<Derived, ArgTypes...>, "Args don't match a Derived constructor");
-
-		#if DO_CHECK
-			PRAGMA_DISABLE_EXPERIMENTAL_WARNINGS
-			check(Metasound::IsChildOfByName(InDerivedTypeName, GetMetasoundDataTypeName<DataType>()));
-			PRAGMA_ENABLE_EXPERIMENTAL_WARNINGS
-		#endif
-			
-			return TDataValueReference(
-				EDataRefShouldConstruct::NewObject,
-				MetasoundDataReferencePrivate::MakeRefOf<DataType, Derived>(Forward<ArgTypes>(Args)...),
-				InDerivedTypeName,
-				InDerivedTypeId
-			);
-		}
 		
 		/** Enable copy constructor */
 		TDataValueReference(const TDataValueReference<DataType>& Other) = default;
@@ -507,12 +454,6 @@ namespace Metasound
 		template <typename T>
 		friend TDataWriteReference<T> WriteCast(const TDataReadReference<T>& InReadableRef);
 
-		using FRefType = TDataReference<DataType>::FRefType;
-		TDataWriteReference(EDataRefShouldConstruct, FRefType InRef, FName InTypeName, const void* InTypeId)
-			: TDataReference<DataType>(EDataRefShouldConstruct::NewObject, MoveTemp(InRef), InTypeName, InTypeId)
-		{
-		}
-
 	public:
 		/** This should be used to construct a new DataType object and return this TDataWriteReference as a wrapper around it. */
 		template <typename... ArgTypes>
@@ -523,7 +464,7 @@ namespace Metasound
 		}
 
 		template <typename Derived, typename... ArgTypes>
-		UE_EXPERIMENTAL(5.8, "Polymorphic API is under development")
+		UE_EXPERIMENTAL(5.7, "Polymorphic API is under development")
 		static TDataWriteReference CreateNewDerived(ArgTypes&&... Args)
 		{
 			static_assert(std::is_base_of_v<DataType, Derived>, "Derived must inherit from DataType");
@@ -534,29 +475,6 @@ namespace Metasound
 				MetasoundDataReferencePrivate::MakeRefOf<DataType, Derived>(Forward<ArgTypes>(Args)...)
 			);
 		}
-
-		template <typename Derived, typename... ArgTypes>
-		UE_EXPERIMENTAL(5.8, "Polymorphic API is under development")
-		static TDataWriteReference CreateNewDerivedAs(const FName& InDerivedTypeName,
-													 FMetasoundDataTypeId InDerivedTypeId,
-													 ArgTypes&&... Args)
-		{
-			static_assert(std::is_base_of_v<DataType, Derived>, "Derived must inherit from DataType");
-			static_assert(std::is_constructible_v<Derived, ArgTypes...>, "Args don't match a Derived constructor");
-
-		#if DO_CHECK
-			PRAGMA_DISABLE_EXPERIMENTAL_WARNINGS
-			check(Metasound::IsChildOfByName(InDerivedTypeName, GetMetasoundDataTypeName<DataType>()));
-			PRAGMA_ENABLE_EXPERIMENTAL_WARNINGS
-		#endif
-			
-			return TDataWriteReference(
-				EDataRefShouldConstruct::NewObject,
-				MetasoundDataReferencePrivate::MakeRefOf<DataType, Derived>(Forward<ArgTypes>(Args)...),
-				InDerivedTypeName,
-				InDerivedTypeId
-			);
-		}
 		
 		/** Enable copy constructor */
 		TDataWriteReference(const TDataWriteReference<DataType>& Other) = default;
@@ -634,12 +552,6 @@ namespace Metasound
 		{
 		}
 
-		using FRefType = typename TDataReference<DataType>::FRefType;
-		TDataReadReference(EDataRefShouldConstruct, FRefType InRef, FName InTypeName, const void* InTypeId)
-			: TDataReference<DataType>(EDataRefShouldConstruct::NewObject, MoveTemp(InRef), InTypeName, InTypeId)
-		{
-		}
-
 	public:
 
 		// This should be used to construct a new DataType object and return this TDataReadReference as a wrapper around it.
@@ -651,7 +563,7 @@ namespace Metasound
 		}
 
 		template <typename Derived, typename... ArgTypes>
-		UE_EXPERIMENTAL(5.8, "Polymorphic API is under development")
+		UE_EXPERIMENTAL(5.7, "Polymorphic API is under development")
 		static TDataReadReference CreateNewDerived(ArgTypes&&... Args)
 		{
 			static_assert(std::is_base_of_v<DataType, Derived>, "Derived must inherit from DataType");
@@ -662,27 +574,6 @@ namespace Metasound
 				MetasoundDataReferencePrivate::MakeRefOf<DataType, Derived>(Forward<ArgTypes>(Args)...)
 			);
 		}
-
-		template <typename Derived, typename... ArgTypes>
-		UE_EXPERIMENTAL(5.8, "Polymorphic API is under development")
-		static TDataReadReference CreateNewDerivedAs(const FName& InDerivedTypeName, FMetasoundDataTypeId InDerivedTypeId, ArgTypes&&... Args)
-		{
-			static_assert(std::is_base_of_v<DataType, Derived>, "Derived must inherit from DataType");
-			static_assert(std::is_constructible_v<Derived, ArgTypes...>, "Args don't match a Derived constructor");
-
-		#if DO_CHECK
-			PRAGMA_DISABLE_EXPERIMENTAL_WARNINGS
-			check(Metasound::IsChildOfByName(InDerivedTypeName, GetMetasoundDataTypeName<DataType>()));
-			PRAGMA_ENABLE_EXPERIMENTAL_WARNINGS
-		#endif
-			
-			return TDataReadReference(
-				EDataRefShouldConstruct::NewObject,
-				MetasoundDataReferencePrivate::MakeRefOf<DataType, Derived>(Forward<ArgTypes>(Args)...),
-				InDerivedTypeName,
-				InDerivedTypeId
-			);
-		}
 		
 		TDataReadReference(const TDataReadReference& Other) = default;
 
diff --git a/Engine/Plugins/Runtime/Metasound/Source/MetasoundGraphCore/Public/MetasoundVertexData.h b/Engine/Plugins/Runtime/Metasound/Source/MetasoundGraphCore/Public/MetasoundVertexData.h
--- a/Engine/Plugins/Runtime/Metasound/Source/MetasoundGraphCore/Public/MetasoundVertexData.h
+++ b/Engine/Plugins/Runtime/Metasound/Source/MetasoundGraphCore/Public/MetasoundVertexData.h
@@ -382,8 +382,7 @@ namespace Metasound
 			void Set(const DataReferenceType& InDataReference)
 			{
 				PRAGMA_DISABLE_INTERNAL_WARNINGS
-				checkf(Metasound::IsCastable(InDataReference.GetDataTypeName(), Vertex.DataTypeName),
-					TEXT("Ref=%s, Vertex=%s"), *InDataReference.GetDataTypeName().ToString(), *Vertex.DataTypeName.ToString());
+				check(Metasound::IsCastable( Vertex.DataTypeName, InDataReference.GetDataTypeName()));
 				PRAGMA_ENABLE_INTERNAL_WARNINGS				
 				Data.Emplace(InDataReference);
 				CheckAccessTypeCompatibility(Vertex, *Data);
diff --git a/Engine/Source/Runtime/SignalProcessing/Private/SphericalHarmonicCalculator.cpp b/Engine/Source/Runtime/SignalProcessing/Private/SphericalHarmonicCalculator.cpp
--- a/Engine/Source/Runtime/SignalProcessing/Private/SphericalHarmonicCalculator.cpp
+++ b/Engine/Source/Runtime/SignalProcessing/Private/SphericalHarmonicCalculator.cpp
@@ -753,20 +753,6 @@ void FSphericalHarmonicCalculator::GenerateFirstOrderRotationMatrixGivenDegrees(
 	GenerateFirstOrderRotationMatrixGivenRadians(RotXDegrees * DEG_2_RAD, RotYDegrees * DEG_2_RAD, RotZDegrees * DEG_2_RAD, OutMatrix);
 }
 
-bool FSphericalHarmonicCalculator::IndexToOrderAndDegree(const int32 InChannelIndex, int32& OutOrder, int32& OutDegree)
-{
-	if (InChannelIndex >= 0)
-	{
-		OutOrder = FMath::FloorToInt(FMath::Sqrt(static_cast<float>(InChannelIndex)));
-		for (; OutOrder * (OutOrder + 1) > InChannelIndex; --OutOrder)
-		{
-		}
-		OutDegree = InChannelIndex - OutOrder * (OutOrder + 1) - OutOrder;
-		return true;
-	}
-	return false;
-}
-
 void FSphericalHarmonicCalculator::NormalizeGains(TArrayView<float> Gains)
 {
 	// Normalize. gains.
diff --git a/Engine/Source/Runtime/SignalProcessing/Private/Vbap.cpp b/Engine/Source/Runtime/SignalProcessing/Private/Vbap.cpp
--- a/Engine/Source/Runtime/SignalProcessing/Private/Vbap.cpp
+++ b/Engine/Source/Runtime/SignalProcessing/Private/Vbap.cpp
@@ -235,12 +235,12 @@ namespace Audio
 			else 
 			{
 				// Don't warn about the wrap span if it's too wide.
-				//UE_CLOG(NextIndex != 0, LogTemp, Warning, TEXT("Speaker A=%d,Az=%.2f and Speaker B=%d,Az=%.2f has Width %.2f"),
-				//	Speaker.Speaker.ChannelID,
-				//	Speaker.Speaker.AzimuthDegrees,
-				//	NextSpeaker.Speaker.ChannelID,
-				//	NextSpeaker.Speaker.AzimuthDegrees,
-				//	SpanWidth );
+				UE_CLOG(NextIndex != 0, LogTemp, Warning, TEXT("Speaker A=%d,Az=%.2f and Speaker B=%d,Az=%.2f has Width %.2f"),
+					Speaker.Speaker.ChannelID,
+					Speaker.Speaker.AzimuthDegrees,
+					NextSpeaker.Speaker.ChannelID,
+					NextSpeaker.Speaker.AzimuthDegrees,
+					SpanWidth );
 			}
 		}
 			
