Namespace the example metasound node config customization as its causing name collisions

diff --git a/Engine/Plugins/Experimental/MetasoundExperimental/Source/MetasoundExperimentalEditor/Private/MetasoundExampleNodeDetailsCustomization.cpp b/Engine/Plugins/Experimental/MetasoundExperimental/Source/MetasoundExperimentalEditor/Private/MetasoundExampleNodeDetailsCustomization.cpp
--- a/Engine/Plugins/Experimental/MetasoundExperimental/Source/MetasoundExperimentalEditor/Private/MetasoundExampleNodeDetailsCustomization.cpp
+++ b/Engine/Plugins/Experimental/MetasoundExperimental/Source/MetasoundExperimentalEditor/Private/MetasoundExampleNodeDetailsCustomization.cpp
@@ -14,7 +14,10 @@
 
 #define LOCTEXT_NAMESPACE "MetasoundExperimentalEditor"
 
-const FString StructIdentifierWithDelimiters =  TEXT(".Struct.");
+namespace MetasoundExampleNodeDetailsCustomizationPrivate
+{
+	const FString StructIdentifierWithDelimiters =  TEXT(".Struct.");
+}
 
 FExampleWidgetNodeConfigurationCustomization::FExampleWidgetNodeConfigurationCustomization(TSharedPtr<IPropertyHandle> InStructProperty, TWeakObjectPtr<UMetasoundEditorGraphNode> InNode)
 	: Metasound::Editor::FMetaSoundNodeConfigurationDataDetails(InStructProperty, InNode)
@@ -34,7 +37,8 @@ void FExampleWidgetNodeConfigurationCustomization::OnChildRowAdded(IDetailProper
 	}
 
 	const FString PropertyPath = ChildHandle->GeneratePathToProperty();
-	// Customize a specific member 
+	// Customize a specific member
+	using namespace MetasoundExampleNodeDetailsCustomizationPrivate;
 	if (PropertyPath == StructPropertyPath + StructIdentifierWithDelimiters + GET_MEMBER_NAME_CHECKED(FMetaSoundWidgetExampleNodeConfiguration, MyFloat).ToString())
 	{
 		MyFloatPropertyHandle = ChildHandle;
