Fix balanced macros

diff --git a/Engine/Plugins/Experimental/MetasoundExperimental/Source/MetasoundExperimentalRuntime/Private/MetasoundCatConversionNode.cpp b/Engine/Plugins/Experimental/MetasoundExperimental/Source/MetasoundExperimentalRuntime/Private/MetasoundCatConversionNode.cpp
--- a/Engine/Plugins/Experimental/MetasoundExperimental/Source/MetasoundExperimentalRuntime/Private/MetasoundCatConversionNode.cpp
+++ b/Engine/Plugins/Experimental/MetasoundExperimental/Source/MetasoundExperimentalRuntime/Private/MetasoundCatConversionNode.cpp
@@ -114,3 +114,5 @@ namespace Metasound
 		FMetasoundFrontendRegistryContainer::Get()->RegisterConversionNode(ConverterNodeKey, CatCast);
 	}
 }
+
+#undef LOCTEXT_NAMESPACE 
